#UE4 #hlslcc Updates from urban

diff --git a/Engine/Source/Developer/ShaderFormatOpenGL/Private/OpenGLShaderCompiler.cpp b/Engine/Source/Developer/ShaderFormatOpenGL/Private/OpenGLShaderCompiler.cpp
--- a/Engine/Source/Developer/ShaderFormatOpenGL/Private/OpenGLShaderCompiler.cpp
+++ b/Engine/Source/Developer/ShaderFormatOpenGL/Private/OpenGLShaderCompiler.cpp
@@ -1486,12 +1486,14 @@ void CompileShader_Windows_OGL(const FShaderCompilerInput& Input,FShaderCompiler
 		}
 
 		FGlslCodeBackend GlslBackEnd(CCFlags);
+        FGlslLanguageSpec GLslLanguageSpec(IsES2Platform(Version));
 		int32 Result = HlslCrossCompile(
 			TCHAR_TO_ANSI(*Input.SourceFilename),
 			TCHAR_TO_ANSI(*PreprocessedShader),
 			TCHAR_TO_ANSI(*Input.EntryPointName),
 			Frequency,
 			&GlslBackEnd,
+            &GLslLanguageSpec,
 			CCFlags,
 			HlslCompilerTarget,
 			&GlslShaderSource,
