#ttp 332327 Don't require a GL context to cross compile GLSL shaders
#branch UE4
#project ShaderFormatOpenGL
#summary The OpenGL shader compiler doesn't need to compile the GLSL within an actual GL context.
#change Remove the bCompileMicrocode option from CompileShader_Windows_OGL.
#change Put the compilation with an actual GL context behind #if DEBUG_HLSLCC - it can be useful when debugging HLSLCC but isn't useful otherwise.

diff --git a/Engine/Source/Developer/ShaderFormatOpenGL/Private/OpenGLShaderCompiler.cpp b/Engine/Source/Developer/ShaderFormatOpenGL/Private/OpenGLShaderCompiler.cpp
--- a/Engine/Source/Developer/ShaderFormatOpenGL/Private/OpenGLShaderCompiler.cpp
+++ b/Engine/Source/Developer/ShaderFormatOpenGL/Private/OpenGLShaderCompiler.cpp
@@ -1360,7 +1360,7 @@ static FString CreateCommandLineHLSLCC( const FString& ShaderFile, const FString
  * @param Input - The input shader code and environment.
  * @param Output - Contains shader compilation results upon return.
  */
-void CompileShader_Windows_OGL(const FShaderCompilerInput& Input,FShaderCompilerOutput& Output,const FString& WorkingDirectory, GLSLVersion Version, bool bCompileMicrocode)
+void CompileShader_Windows_OGL(const FShaderCompilerInput& Input,FShaderCompilerOutput& Output,const FString& WorkingDirectory, GLSLVersion Version)
 {
 	FString PreprocessedShader;
 	FShaderCompilerDefinitions AdditionalDefines;
@@ -1530,27 +1530,24 @@ void CompileShader_Windows_OGL(const FShaderCompilerInput& Input,FShaderCompiler
 				}
 			}
 
-			if (bCompileMicrocode)
+#if VALIDATE_GLSL_WITH_DRIVER
+			PrecompileShader(Output, Input, GlslShaderSource, Version, Frequency);
+			if (Output.bSucceeded == false)
 			{
-				PrecompileShader(Output, Input, GlslShaderSource, Version, Frequency);
-				if (Output.bSucceeded == false)
-				{
 #if DUMP_HLSCLCC_SHADERS
-					DumpFile.SetFilename( *DumpGLSL);
-					for(int i = 0; i < Output.Errors.Num(); ++i)
-					{
-						DumpFile.Logf(TEXT("%s"), *Output.Errors[i].GetErrorString());
-					}
-					DumpFile.Flush();
-#endif
+				DumpFile.SetFilename( *DumpGLSL);
+				for(int i = 0; i < Output.Errors.Num(); ++i)
+				{
+					DumpFile.Logf(TEXT("%s"), *Output.Errors[i].GetErrorString());
 				}
+				DumpFile.Flush();
+#endif
 			}
-			else
-			{
-				int32 SourceLen = FCStringAnsi::Strlen(GlslShaderSource);
-				Output.Target = Input.Target;
-				BuildShaderOutput(Output, GlslShaderSource, SourceLen, Version);
-			}
+#else // VALIDATE_GLSL_WITH_DRIVER
+			int32 SourceLen = FCStringAnsi::Strlen(GlslShaderSource);
+			Output.Target = Input.Target;
+			BuildShaderOutput(Output, GlslShaderSource, SourceLen, Version);
+#endif // VALIDATE_GLSL_WITH_DRIVER
 		}
 		else
 		{
diff --git a/Engine/Source/Developer/ShaderFormatOpenGL/Private/ShaderFormatOpenGL.h b/Engine/Source/Developer/ShaderFormatOpenGL/Private/ShaderFormatOpenGL.h
--- a/Engine/Source/Developer/ShaderFormatOpenGL/Private/ShaderFormatOpenGL.h
+++ b/Engine/Source/Developer/ShaderFormatOpenGL/Private/ShaderFormatOpenGL.h
@@ -13,4 +13,4 @@ enum GLSLVersion
 	GLSL_150_MAC, // Apple only
 };
 
-extern void CompileShader_Windows_OGL(const struct FShaderCompilerInput& Input,struct FShaderCompilerOutput& Output,const class FString& WorkingDirectory, GLSLVersion Version, bool bCompileMicrocode = true);
+extern void CompileShader_Windows_OGL(const struct FShaderCompilerInput& Input,struct FShaderCompilerOutput& Output,const class FString& WorkingDirectory, GLSLVersion Version);
