Fix GL ES2 emulation shader compiler issue

diff --git a/Engine/Source/Developer/ShaderFormatOpenGL/Private/OpenGLShaderCompiler.cpp b/Engine/Source/Developer/ShaderFormatOpenGL/Private/OpenGLShaderCompiler.cpp
--- a/Engine/Source/Developer/ShaderFormatOpenGL/Private/OpenGLShaderCompiler.cpp
+++ b/Engine/Source/Developer/ShaderFormatOpenGL/Private/OpenGLShaderCompiler.cpp
@@ -1486,7 +1486,7 @@ void CompileShader_Windows_OGL(const FShaderCompilerInput& Input,FShaderCompiler
 		}
 
 		FGlslCodeBackend GlslBackEnd(CCFlags);
-		FGlslLanguageSpec GLslLanguageSpec(IsES2Platform(Version));
+		FGlslLanguageSpec GLslLanguageSpec(IsES2Platform(Version) && !IsPCES2Platform(Version));
 		int32 Result = HlslCrossCompile(
 			TCHAR_TO_ANSI(*Input.SourceFilename),
 			TCHAR_TO_ANSI(*PreprocessedShader),
diff --git a/Engine/Source/Developer/ShaderFormatOpenGL/Private/ShaderFormatOpenGL.cpp b/Engine/Source/Developer/ShaderFormatOpenGL/Private/ShaderFormatOpenGL.cpp
--- a/Engine/Source/Developer/ShaderFormatOpenGL/Private/ShaderFormatOpenGL.cpp
+++ b/Engine/Source/Developer/ShaderFormatOpenGL/Private/ShaderFormatOpenGL.cpp
@@ -25,7 +25,7 @@ class FShaderFormatGLSL : public IShaderFormat
 		UE_SHADER_GLSL_150_MAC_VER = 1,
 		UE_SHADER_GLSL_430_VER = 51,
 		UE_SHADER_GLSL_ES2_VER = 10,
-		UE_SHADER_GLSL_150ES2_VER = 11,
+		UE_SHADER_GLSL_150ES2_VER = 12,
 		UE_SHADER_GLSL_ES2_VER_WEBGL  = 11,
 		UE_SHADER_GLSL_ES2_IOS_VER  = 1,
 	}; 
@@ -37,7 +37,7 @@ class FShaderFormatGLSL : public IShaderFormat
 				Format == NAME_GLSL_430 || 
 				Format == NAME_GLSL_ES2 || 
 				Format == NAME_GLSL_150_ES2 ||
-                Format == NAME_GLSL_ES2_WEBGL ||
+				Format == NAME_GLSL_ES2_WEBGL ||
 				Format == NAME_GLSL_ES2_IOS
 			);
 	}
