Increase iteration times on hlslcc; remove -ios as it's not used

diff --git a/Engine/Source/Developer/ShaderFormatOpenGL/Private/OpenGLShaderCompiler.cpp b/Engine/Source/Developer/ShaderFormatOpenGL/Private/OpenGLShaderCompiler.cpp
--- a/Engine/Source/Developer/ShaderFormatOpenGL/Private/OpenGLShaderCompiler.cpp
+++ b/Engine/Source/Developer/ShaderFormatOpenGL/Private/OpenGLShaderCompiler.cpp
@@ -1342,7 +1342,7 @@ static FString CreateCommandLineHLSLCC( const FString& ShaderFile, const FString
 			break;
 
 		case GLSL_ES2_IOS:
-			VersionSwitch = TEXT(" -ios -es2");
+			VersionSwitch = TEXT(" -es2");
 			break;
 
 		default:
