Support for Linux in shader compiler.

- Part of PR #306 by amigo and #ue4linux folks.
- Undoing my previous change with cross-compilation define (too dangerous).

diff --git a/Engine/Source/Developer/ShaderFormatOpenGL/Private/OpenGLShaderCompiler.cpp b/Engine/Source/Developer/ShaderFormatOpenGL/Private/OpenGLShaderCompiler.cpp
--- a/Engine/Source/Developer/ShaderFormatOpenGL/Private/OpenGLShaderCompiler.cpp
+++ b/Engine/Source/Developer/ShaderFormatOpenGL/Private/OpenGLShaderCompiler.cpp
@@ -25,9 +25,19 @@
 	#include <GL/wglext.h>
 #include "HideWindowsPlatformTypes.h"
 #elif PLATFORM_LINUX
-	#define GL_GLEXT_PROTOTYPES
+	#define GL_GLEXT_PROTOTYPES 1
 	#include <GL/glcorearb.h>
 	#include <GL/glext.h>
+	#include "SDL.h"
+	GLAPI GLuint APIENTRY glCreateShader (GLenum type);
+	GLAPI void APIENTRY glShaderSource (GLuint shader, GLsizei count, const GLchar* const *string, const GLint *length);
+	typedef SDL_Window*		SDL_HWindow;
+	typedef SDL_GLContext	SDL_HGLContext;
+	struct FPlatformOpenGLContext
+	{
+		SDL_HWindow		hWnd;
+		SDL_HGLContext	hGLContext;		//	this is a (void*) pointer
+	};
 #elif PLATFORM_MAC
 	#include <OpenGL/OpenGL.h>
 	#include <OpenGL/gl3.h>
@@ -261,13 +271,107 @@ static void PlatformReleaseOpenGL(void* ContextPtr, void* PrevContextPtr)
 	wglMakeCurrent((HDC)ContextPtr, (HGLRC)PrevContextPtr);
 }
 #elif PLATFORM_LINUX
+static void _PlatformCreateDummyGLWindow(FPlatformOpenGLContext *OutContext)
+{
+	static bool bInitializedWindowClass = false;
+
+	// Create a dummy window.
+	OutContext->hWnd = SDL_CreateWindow(NULL,
+		0, 0, 1, 1,
+		SDL_WINDOW_OPENGL | SDL_WINDOW_BORDERLESS | SDL_WINDOW_HIDDEN);
+}
+
+static void _PlatformCreateOpenGLContextCore(FPlatformOpenGLContext* OutContext)
+{
+	check(OutContext);
+	SDL_HWindow PrevWindow = SDL_GL_GetCurrentWindow();
+	SDL_HGLContext PrevContext = SDL_GL_GetCurrentContext();
+
+	OutContext->hGLContext = SDL_GL_CreateContext(OutContext->hWnd);
+	SDL_GL_MakeCurrent(PrevWindow, PrevContext);
+}
+
+static void _ContextMakeCurrent(SDL_HWindow hWnd, SDL_HGLContext hGLDC)
+{
+	GLint Result = SDL_GL_MakeCurrent( hWnd, hGLDC );
+	check(!Result);
+}
+
 static void PlatformInitOpenGL(void*& ContextPtr, void*& PrevContextPtr, int InMajorVersion, int InMinorVersion)
 {
-  unimplemented();
+	static bool bInitialized = (SDL_GL_GetCurrentWindow() != NULL) && (SDL_GL_GetCurrentContext() != NULL);
+
+	if (!bInitialized)
+	{
+		check(InMajorVersion > 3 || (InMajorVersion == 3 && InMinorVersion >= 2));
+		if (SDL_WasInit(0) == 0)
+		{
+			SDL_Init(SDL_INIT_VIDEO);
+		}
+		else
+		{
+			Uint32 InitializedSubsystemsMask = SDL_WasInit(SDL_INIT_EVERYTHING);
+			if ((InitializedSubsystemsMask & SDL_INIT_VIDEO) == 0)
+			{
+				SDL_InitSubSystem(SDL_INIT_VIDEO);
+			}
+		}
+
+		if (SDL_GL_LoadLibrary(NULL))
+		{
+			UE_LOG(LogOpenGLShaderCompiler, Fatal, TEXT("Unable to dynamically load libGL: %s"), ANSI_TO_TCHAR(SDL_GetError()));
+		}
+
+		if	(SDL_GL_SetAttribute( SDL_GL_CONTEXT_MAJOR_VERSION, InMajorVersion))
+		{
+			UE_LOG(LogOpenGLShaderCompiler, Fatal, TEXT("Failed to set GL major version: %s"), ANSI_TO_TCHAR(SDL_GetError()));
+		}
+
+		if	(SDL_GL_SetAttribute( SDL_GL_CONTEXT_MINOR_VERSION, InMinorVersion))
+		{
+			UE_LOG(LogOpenGLShaderCompiler, Fatal, TEXT("Failed to set GL minor version: %s"), ANSI_TO_TCHAR(SDL_GetError()));
+		}
+
+		if	(SDL_GL_SetAttribute( SDL_GL_CONTEXT_FLAGS, SDL_GL_CONTEXT_FORWARD_COMPATIBLE_FLAG))
+		{
+			UE_LOG(LogOpenGLShaderCompiler, Fatal, TEXT("Failed to set GL flags: %s"), ANSI_TO_TCHAR(SDL_GetError()));
+		}
+
+		if	(SDL_GL_SetAttribute( SDL_GL_CONTEXT_PROFILE_MASK, SDL_GL_CONTEXT_PROFILE_CORE))
+		{
+			UE_LOG(LogOpenGLShaderCompiler, Fatal, TEXT("Failed to set GL mask/profile: %s"), ANSI_TO_TCHAR(SDL_GetError()));
+		}
+
+		// Create a dummy context to verify opengl support.
+		FPlatformOpenGLContext DummyContext;
+		_PlatformCreateDummyGLWindow(&DummyContext);
+		_PlatformCreateOpenGLContextCore(&DummyContext);
+
+		if (DummyContext.hGLContext)
+		{
+			_ContextMakeCurrent(DummyContext.hWnd, DummyContext.hGLContext);
+		}
+		else
+		{
+			UE_LOG(LogOpenGLShaderCompiler, Fatal, TEXT("OpenGL %d.%d not supported by driver"), InMajorVersion, InMinorVersion);
+			return;
+		}
+
+		PrevContextPtr = NULL;
+		ContextPtr = DummyContext.hGLContext;
+		bInitialized = true;
+	}
+
+	PrevContextPtr = reinterpret_cast<void*>(SDL_GL_GetCurrentContext());
+	SDL_HGLContext NewContext = SDL_GL_CreateContext(SDL_GL_GetCurrentWindow());
+	SDL_GL_MakeCurrent(SDL_GL_GetCurrentWindow(), NewContext);
+	ContextPtr = reinterpret_cast<void*>(NewContext);
 }
+
 static void PlatformReleaseOpenGL(void* ContextPtr, void* PrevContextPtr)
 {
-  unimplemented();
+	SDL_GL_MakeCurrent(SDL_GL_GetCurrentWindow(), reinterpret_cast<SDL_HGLContext>(PrevContextPtr));
+	SDL_GL_DeleteContext(reinterpret_cast<SDL_HGLContext>(ContextPtr));
 }
 #elif PLATFORM_MAC
 static void PlatformInitOpenGL(void*& ContextPtr, void*& PrevContextPtr, int InMajorVersion, int InMinorVersion)
diff --git a/Engine/Source/Programs/UnrealBuildTool/Linux/UEBuildLinux.cs b/Engine/Source/Programs/UnrealBuildTool/Linux/UEBuildLinux.cs
--- a/Engine/Source/Programs/UnrealBuildTool/Linux/UEBuildLinux.cs
+++ b/Engine/Source/Programs/UnrealBuildTool/Linux/UEBuildLinux.cs
@@ -340,7 +340,6 @@ namespace UnrealBuildTool
         {
             InBuildTarget.GlobalCompileEnvironment.Config.Definitions.Add("PLATFORM_LINUX=1");
             InBuildTarget.GlobalCompileEnvironment.Config.Definitions.Add("LINUX=1");
-            InBuildTarget.GlobalCompileEnvironment.Config.Definitions.Add(String.Format("CROSS_COMPILATION={0}", ExternalExecution.GetRuntimePlatform() != UnrealTargetPlatform.Linux ? 1 : 0));
 
             InBuildTarget.GlobalCompileEnvironment.Config.Definitions.Add("WITH_DATABASE_SUPPORT=0");		//@todo linux: valid?
 
