UE4 - Move glsl backend out of ThirdParty into ShaderFormatOpenGL

diff --git a/Engine/Source/Developer/ShaderCompilerCommon/Private/ShaderCompilerCommon.cpp b/Engine/Source/Developer/ShaderCompilerCommon/Private/ShaderCompilerCommon.cpp
--- a/Engine/Source/Developer/ShaderCompilerCommon/Private/ShaderCompilerCommon.cpp
+++ b/Engine/Source/Developer/ShaderCompilerCommon/Private/ShaderCompilerCommon.cpp
@@ -183,18 +183,11 @@ bool RemoveUniformBuffersFromSource(FString& SourceCode)
 }
 
 
-FString CreateCrossCompilerBatchFileContents(const FString& ShaderFile, const FString& OutputFile, const FString& FrequencySwitch, const FString& EntryPoint, const FString& VersionSwitch, bool bUseCCT, const FString& ExtraArguments)
+FString CreateCrossCompilerBatchFileContents(const FString& ShaderFile, const FString& OutputFile, const FString& FrequencySwitch, const FString& EntryPoint, const FString& VersionSwitch, const FString& ExtraArguments)
 {
 	FString BatchFile = TEXT("@echo off");
 	BatchFile += TEXT("\nif defined ue.hlslcc GOTO DONE\nset ue.hlslcc=");
-	if (bUseCCT)
-	{
-		BatchFile += FPaths::RootDir() / TEXT("Engine\\Source\\Binaries\\Win64\\CrossCompilerTool.exe");
-	}
-	else
-	{
-		BatchFile += FPaths::RootDir() / TEXT("Engine\\Source\\ThirdParty\\hlslcc\\hlslcc\\bin\\Win64\\VS2013\\hlslcc_64.exe");
-	}
+	BatchFile += FPaths::RootDir() / TEXT("Engine\\Source\\Binaries\\Win64\\CrossCompilerTool.exe");
 	BatchFile += TEXT("\n\n:DONE\n%ue.hlslcc% ");
 	BatchFile += FString::Printf(TEXT("\"%s\" -o=\"%s\" %s -entry=%s %s %s"), *ShaderFile, *OutputFile, *FrequencySwitch, *EntryPoint, *VersionSwitch, *ExtraArguments);
 	BatchFile += TEXT("\npause\n");
diff --git a/Engine/Source/Developer/ShaderCompilerCommon/Public/ShaderCompilerCommon.h b/Engine/Source/Developer/ShaderCompilerCommon/Public/ShaderCompilerCommon.h
--- a/Engine/Source/Developer/ShaderCompilerCommon/Public/ShaderCompilerCommon.h
+++ b/Engine/Source/Developer/ShaderCompilerCommon/Public/ShaderCompilerCommon.h
@@ -70,5 +70,4 @@ extern SHADERCOMPILERCOMMON_API FString CreateCrossCompilerBatchFileContents(
 											const FString& FrequencySwitch,
 											const FString& EntryPoint,
 											const FString& VersionSwitch,
-											bool bUseCCT,
 											const FString& ExtraArguments = TEXT(""));
diff --git a/Engine/Source/Developer/ShaderFormatOpenGL/Private/OpenGLShaderCompiler.cpp b/Engine/Source/Developer/ShaderFormatOpenGL/Private/OpenGLShaderCompiler.cpp
--- a/Engine/Source/Developer/ShaderFormatOpenGL/Private/OpenGLShaderCompiler.cpp
+++ b/Engine/Source/Developer/ShaderFormatOpenGL/Private/OpenGLShaderCompiler.cpp
@@ -17,7 +17,7 @@
 #include "ShaderPreprocessor.h"
 #include "ShaderCompilerCommon.h"
 #include "hlslcc.h"
-#include "glsl/ir_gen_glsl.h"
+#include "ir_gen_glsl.h"
 #if PLATFORM_WINDOWS
 #include "AllowWindowsPlatformTypes.h"
 	#include <GL/glcorearb.h>
@@ -1496,7 +1496,7 @@ static FString CreateCrossCompilerBatchFile( const FString& ShaderFile, const FS
 	}
 
 	const TCHAR* ApplyCSE = (CCFlags & HLSLCC_ApplyCommonSubexpressionElimination) != 0 ? TEXT("-cse") : TEXT("");
-	return CreateCrossCompilerBatchFileContents(ShaderFile, OutputFile, FrequencySwitch, EntryPoint, VersionSwitch, false, ApplyCSE);
+	return CreateCrossCompilerBatchFileContents(ShaderFile, OutputFile, FrequencySwitch, EntryPoint, VersionSwitch, ApplyCSE);
 }
 
 /**
diff --git a/Engine/Source/Developer/ShaderFormatOpenGL/Private/ir_gen_glsl.cpp b/Engine/Source/Developer/ShaderFormatOpenGL/Private/ir_gen_glsl.cpp
new file mode 100644
--- /dev/null
+++ b/Engine/Source/Developer/ShaderFormatOpenGL/Private/ir_gen_glsl.cpp
@@ -0,0 +1,4975 @@
+// Copyright 1998-2014 Epic Games, Inc. All Rights Reserved.
+
+// This code is largely based on that in ir_print_glsl_visitor.cpp from
+// glsl-optimizer.
+// [REDACTED_PROFILE_URL]/glsl-optimizer
+// The license for glsl-optimizer is reproduced below:
+
+/*
+	GLSL Optimizer is licensed according to the terms of the MIT license:
+
+	Copyright (C) 1999-2007  Brian Paul   All Rights Reserved.
+	Copyright (C) 2010-2011  Unity Technologies All Rights Reserved.
+
+	Permission is hereby granted, free of charge, to any person obtaining a
+	copy of this software and associated documentation files (the "Software"),
+	to deal in the Software without restriction, including without limitation
+	the rights to use, copy, modify, merge, publish, distribute, sublicense,
+	and/or sell copies of the Software, and to permit persons to whom the
+	Software is furnished to do so, subject to the following conditions:
+
+	The above copyright notice and this permission notice shall be included
+	in all copies or substantial portions of the Software.
+
+	THE SOFTWARE IS PROVIDED "AS IS", WITHOUT WARRANTY OF ANY KIND, EXPRESS
+	OR IMPLIED, INCLUDING BUT NOT LIMITED TO THE WARRANTIES OF MERCHANTABILITY,
+	FITNESS FOR A PARTICULAR PURPOSE AND NONINFRINGEMENT.  IN NO EVENT SHALL
+	BRIAN PAUL BE LIABLE FOR ANY CLAIM, DAMAGES OR OTHER LIABILITY, WHETHER IN
+	AN ACTION OF CONTRACT, TORT OR OTHERWISE, ARISING FROM, OUT OF OR IN
+	CONNECTION WITH THE SOFTWARE OR THE USE OR OTHER DEALINGS IN THE SOFTWARE.
+*/
+
+#include "ShaderFormatOpenGL.h"
+#include "Core.h"
+#include "hlslcc.h"
+#include "hlslcc_private.h"
+#include "ir_gen_glsl.h"
+#include "mesa/compiler.h"
+#include "mesa/glsl_parser_extras.h"
+#include "mesa/hash_table.h"
+#include "mesa/ir_rvalue_visitor.h"
+#include "PackUniformBuffers.h"
+#include "IRDump.h"
+//@todo-rco: Remove STL!
+#include <sstream>
+//#define OPTIMIZE_ANON_STRUCTURES_OUT
+// We can't optimize them out presently, because apparently Windows Radeon
+// OpenGL driver chokes on valid GLSL code then.
+
+#if !PLATFORM_WINDOWS
+#define _strnicmp strnicmp
+#define _stricmp stricmp
+#endif
+
+static bool GDefaultPrecisionIsHalf = true;
+
+static inline std::string FixHlslName(const glsl_type* Type)
+{
+	check(Type->is_image() || Type->is_vector() || Type->is_numeric() || Type->is_void() || Type->is_sampler() || Type->is_scalar());
+	std::string Name = Type->name;
+	if (Type == glsl_type::half_type)
+	{
+		return "float";
+	}
+	else if (Type == glsl_type::half2_type)
+	{
+		return "vec2";
+	}
+	else if (Type == glsl_type::half3_type)
+	{
+		return "vec3";
+	}
+	else if (Type == glsl_type::half4_type)
+	{
+		return "vec4";
+	}
+
+	return Name;
+}
+
+/**
+ * This table must match the ir_expression_operation enum.
+ */
+static const char * const GLSLExpressionTable[ir_opcode_count][4] =
+	{
+	{ "(~", ")", "", "" }, // ir_unop_bit_not,
+	{ "not(", ")", "", "!" }, // ir_unop_logic_not,
+	{ "(-", ")", "", "" }, // ir_unop_neg,
+	{ "abs(", ")", "", "" }, // ir_unop_abs,
+	{ "sign(", ")", "", "" }, // ir_unop_sign,
+	{ "(1.0/(", "))", "", "" }, // ir_unop_rcp,
+	{ "inversesqrt(", ")", "", "" }, // ir_unop_rsq,
+	{ "sqrt(", ")", "", "" }, // ir_unop_sqrt,
+	{ "exp(", ")", "", "" }, // ir_unop_exp,      /**< Log base e on gentype */
+	{ "log(", ")", "", "" }, // ir_unop_log,	     /**< Natural log on gentype */
+	{ "exp2(", ")", "", "" }, // ir_unop_exp2,
+	{ "log2(", ")", "", "" }, // ir_unop_log2,
+	{ "int(", ")", "", "" }, // ir_unop_f2i,      /**< Float-to-integer conversion. */
+	{ "float(", ")", "", "" }, // ir_unop_i2f,      /**< Integer-to-float conversion. */
+	{ "bool(", ")", "", "" }, // ir_unop_f2b,      /**< Float-to-boolean conversion */
+	{ "float(", ")", "", "" }, // ir_unop_b2f,      /**< Boolean-to-float conversion */
+	{ "bool(", ")", "", "" }, // ir_unop_i2b,      /**< int-to-boolean conversion */
+	{ "int(", ")", "", "" }, // ir_unop_b2i,      /**< Boolean-to-int conversion */
+	{ "uint(", ")", "", "" }, // ir_unop_b2u,
+	{ "bool(", ")", "", "" }, // ir_unop_u2b,
+	{ "uint(", ")", "", "" }, // ir_unop_f2u,
+	{ "float(", ")", "", "" }, // ir_unop_u2f,      /**< Unsigned-to-float conversion. */
+	{ "uint(", ")", "", "" }, // ir_unop_i2u,      /**< Integer-to-unsigned conversion. */
+	{ "int(", ")", "", "" }, // ir_unop_u2i,      /**< Unsigned-to-integer conversion. */
+	{ "int(", ")", "", "" }, // ir_unop_h2i,
+	{ "float(", ")", "", "" }, // ir_unop_i2h,
+	{ "(", ")", "", "" }, // ir_unop_h2f,
+	{ "(", ")", "", "" }, // ir_unop_f2h,
+	{ "bool(", ")", "", "" }, // ir_unop_h2b,
+	{ "float(", ")", "", "" }, // ir_unop_b2h,
+	{ "uint(", ")", "", "" }, // ir_unop_h2u,
+	{ "uint(", ")", "", "" }, // ir_unop_u2h,
+	{ "transpose(", ")", "", "" }, // ir_unop_transpose
+	{ "any(", ")", "", "" }, // ir_unop_any,
+	{ "all(", ")", "", "" }, // ir_unop_all,
+
+	/**
+	* \name Unary floating-point rounding operations.
+	*/
+	/*@{*/
+	{ "trunc(", ")", "", "" }, // ir_unop_trunc,
+	{ "ceil(", ")", "", "" }, // ir_unop_ceil,
+	{ "floor(", ")", "", "" }, // ir_unop_floor,
+	{ "fract(", ")", "", "" }, // ir_unop_fract,
+	{ "round(", ")", "", "" }, // ir_unop_round,
+	/*@}*/
+
+	/**
+	* \name Trigonometric operations.
+	*/
+	/*@{*/
+	{ "sin(", ")", "", "" }, // ir_unop_sin,
+	{ "cos(", ")", "", "" }, // ir_unop_cos,
+	{ "tan(", ")", "", "" }, // ir_unop_tan,
+	{ "asin(", ")", "", "" }, // ir_unop_asin,
+	{ "acos(", ")", "", "" }, // ir_unop_acos,
+	{ "atan(", ")", "", "" }, // ir_unop_atan,
+	{ "sinh(", ")", "", "" }, // ir_unop_sinh,
+	{ "cosh(", ")", "", "" }, // ir_unop_cosh,
+	{ "tanh(", ")", "", "" }, // ir_unop_tanh,
+	/*@}*/
+
+	/**
+	* \name Normalize.
+	*/
+	/*@{*/
+	{ "normalize(", ")", "", "" }, // ir_unop_normalize,
+	/*@}*/
+
+	/**
+	* \name Partial derivatives.
+	*/
+	/*@{*/
+	{ "dFdx(", ")", "", "" }, // ir_unop_dFdx,
+	{ "dFdy(", ")", "", "" }, // ir_unop_dFdy,
+	/*@}*/
+
+	{ "isnan(", ")", "", "" }, // ir_unop_isnan,
+	{ "isinf(", ")", "", "" }, // ir_unop_isinf,
+
+	{ "floatBitsToUint(", ")", "", "" }, // ir_unop_fasu,
+	{ "floatBitsToInt(", ")", "", "" }, // ir_unop_fasi,
+	{ "intBitsToFloat(", ")", "", "" }, // ir_unop_iasf,
+	{ "uintBitsToFloat(", ")", "", "" }, // ir_unop_uasf,
+
+	{ "bitfieldReverse(", ")", "", "" }, // ir_unop_bitreverse,
+	{ "bitCount(", ")", "", "" }, // ir_unop_bitcount,
+	{ "findMSB(", ")", "", "" }, // ir_unop_msb,
+	{ "findLSB(", ")", "", "" }, // ir_unop_lsb,
+
+	{ "ERROR_NO_NOISE_FUNCS(", ")", "", "" }, // ir_unop_noise,
+
+	{ "(", "+", ")", "" }, // ir_binop_add,
+	{ "(", "-", ")", "" }, // ir_binop_sub,
+	{ "(", "*", ")", "" }, // ir_binop_mul,
+	{ "(", "/", ")", "" }, // ir_binop_div,
+
+	/**
+	* Takes one of two combinations of arguments:
+	*
+	* - mod(vecN, vecN)
+	* - mod(vecN, float)
+	*
+	* Does not take integer types.
+	*/
+	{ "mod(", ",", ")", "%" }, // ir_binop_mod,
+	{ "modf(", ",", ")", "" }, // ir_binop_modf,
+
+	{ "step(", ",", ")", "" }, // ir_binop_step,
+
+	/**
+	* \name Binary comparison operators which return a boolean vector.
+	* The type of both operands must be equal.
+	*/
+	/*@{*/
+	{ "lessThan(", ",", ")", "<" }, // ir_binop_less,
+	{ "greaterThan(", ",", ")", ">" }, // ir_binop_greater,
+	{ "lessThanEqual(", ",", ")", "<=" }, // ir_binop_lequal,
+	{ "greaterThanEqual(", ",", ")", ">=" }, // ir_binop_gequal,
+	{ "equal(", ",", ")", "==" }, // ir_binop_equal,
+	{ "notEqual(", ",", ")", "!=" }, // ir_binop_nequal,
+	/**
+	* Returns single boolean for whether all components of operands[0]
+	* equal the components of operands[1].
+	*/
+	{ "(", "==", ")", "" }, // ir_binop_all_equal,
+	/**
+	* Returns single boolean for whether any component of operands[0]
+	* is not equal to the corresponding component of operands[1].
+	*/
+	{ "(", "!=", ")", "" }, // ir_binop_any_nequal,
+	/*@}*/
+
+	/**
+	* \name Bit-wise binary operations.
+	*/
+	/*@{*/
+	{ "(", "<<", ")", "" }, // ir_binop_lshift,
+	{ "(", ">>", ")", "" }, // ir_binop_rshift,
+	{ "(", "&", ")", "" }, // ir_binop_bit_and,
+	{ "(", "^", ")", "" }, // ir_binop_bit_xor,
+	{ "(", "|", ")", "" }, // ir_binop_bit_or,
+	/*@}*/
+
+	{ "bvec%d(uvec%d(", ")*uvec%d(", "))", "&&" }, // ir_binop_logic_and,
+	{ "bvec%d(abs(ivec%d(", ")+ivec%d(", ")))", "^^" }, // ir_binop_logic_xor,
+	{ "bvec%d(uvec%d(", ")+uvec%d(", "))", "||" }, // ir_binop_logic_or,
+
+	{ "dot(", ",", ")", "" }, // ir_binop_dot,
+	{ "cross(", ",", ")", "" }, // ir_binop_cross,
+	{ "min(", ",", ")", "" }, // ir_binop_min,
+	{ "max(", ",", ")", "" }, // ir_binop_max,
+	{ "atan(", ",", ")", "" },
+	{ "pow(", ",", ")", "" }, // ir_binop_pow,
+
+	{ "mix(", ",", ",", ")" }, // ir_ternop_lerp,
+	{ "smoothstep(", ",", ",", ")" }, // ir_ternop_smoothstep,
+	{ "clamp(", ",", ",", ")" }, // ir_ternop_clamp,
+
+	{ "ERROR_QUADOP_VECTOR(", ",", ")" }, // ir_quadop_vector,
+};
+
+static const char* OutputStreamTypeStrings[4] = {
+	"!invalid!",
+	"points",
+	"line_strip",
+	"triangle_strip"
+};
+
+static const char* GeometryInputStrings[6] = {
+	"!invalid!",
+	"points",
+	"lines",
+	"line_adjacency",
+	"triangles",
+	"triangles_adjacency"
+};
+
+static const char* DomainStrings[4] = {
+	"!invalid!",
+	"triangles",
+	"quads",
+	"isolines",
+};
+
+static const char* PartitioningStrings[5] = {
+	"!invalid!",
+	"equal_spacing",
+	"fractional_even_spacing",
+	"fractional_odd_spacing",
+	"pow2",
+};
+
+static const char* OutputTopologyStrings[5] = {
+	"!invalid!",
+	"point_needs_to_be_fixed",
+	"line_needs_to_be_fixed",
+	"cw",
+	"ccw",
+};
+
+static_assert((sizeof(GLSLExpressionTable) / sizeof(GLSLExpressionTable[0])) == ir_opcode_count, "GLSLExpressionTableSizeMismatch");
+
+struct SDMARange
+{
+	unsigned SourceCB;
+	unsigned SourceOffset;
+	unsigned Size;
+	unsigned DestCBIndex;
+	unsigned DestCBPrecision;
+	unsigned DestOffset;
+
+	bool operator <(SDMARange const & Other) const
+	{
+		if (SourceCB == Other.SourceCB)
+		{
+			return SourceOffset < Other.SourceOffset;
+		}
+
+		return SourceCB < Other.SourceCB;
+	}
+};
+typedef std::list<SDMARange> TDMARangeList;
+typedef std::map<unsigned, TDMARangeList> TCBDMARangeMap;
+
+
+static void InsertRange( TCBDMARangeMap& CBAllRanges, unsigned SourceCB, unsigned SourceOffset, unsigned Size, unsigned DestCBIndex, unsigned DestCBPrecision, unsigned DestOffset ) 
+{
+	check(SourceCB < (1 << 12));
+	check(DestCBIndex < (1 << 12));
+	check(DestCBPrecision < (1 << 8));
+	unsigned SourceDestCBKey = (SourceCB << 20) | (DestCBIndex << 8) | DestCBPrecision;
+	SDMARange Range = { SourceCB, SourceOffset, Size, DestCBIndex, DestCBPrecision, DestOffset };
+
+	TDMARangeList& CBRanges = CBAllRanges[SourceDestCBKey];
+//printf("* InsertRange: %08x\t%d:%d - %d:%c:%d:%d\n", SourceDestCBKey, SourceCB, SourceOffset, DestCBIndex, DestCBPrecision, DestOffset, Size);
+	if (CBRanges.empty())
+	{
+		CBRanges.push_back(Range);
+	}
+	else
+	{
+		TDMARangeList::iterator Prev = CBRanges.end();
+		bool bAdded = false;
+		for (auto Iter = CBRanges.begin(); Iter != CBRanges.end(); ++Iter)
+		{
+			if (SourceOffset + Size <= Iter->SourceOffset)
+			{
+				if (Prev == CBRanges.end())
+				{
+					CBRanges.push_front(Range);
+				}
+				else
+				{
+					CBRanges.insert(Iter, Range);
+				}
+
+				bAdded = true;
+				break;
+			}
+
+			Prev = Iter;
+		}
+
+		if (!bAdded)
+		{
+			CBRanges.push_back(Range);
+		}
+
+		if (CBRanges.size() > 1)
+		{
+			// Try to merge ranges
+			bool bDirty = false;
+			do
+			{
+				bDirty = false;
+				TDMARangeList NewCBRanges;
+				for (auto Iter = CBRanges.begin(); Iter != CBRanges.end(); ++Iter)
+				{
+					if (Iter == CBRanges.begin())
+					{
+						Prev = CBRanges.begin();
+					}
+					else
+					{
+						if (Prev->SourceOffset + Prev->Size == Iter->SourceOffset && Prev->DestOffset + Prev->Size == Iter->DestOffset)
+						{
+							SDMARange Merged = *Prev;
+							Merged.Size = Prev->Size + Iter->Size;
+							NewCBRanges.pop_back();
+							NewCBRanges.push_back(Merged);
+							++Iter;
+							NewCBRanges.insert(NewCBRanges.end(), Iter, CBRanges.end());
+							bDirty = true;
+							break;
+						}
+					}
+
+					NewCBRanges.push_back(*Iter);
+					Prev = Iter;
+				}
+
+				CBRanges.swap(NewCBRanges);
+			}
+			while (bDirty);
+		}
+	}
+}
+
+static TDMARangeList SortRanges( TCBDMARangeMap& CBRanges ) 
+{
+	TDMARangeList Sorted;
+	for (auto& Pair : CBRanges)
+	{
+		Sorted.insert(Sorted.end(), Pair.second.begin(), Pair.second.end());
+	}
+
+	Sorted.sort();
+
+	return Sorted;
+}
+
+static void DumpSortedRanges(TDMARangeList& SortedRanges)
+{
+	printf("**********************************\n");
+	for (auto& o : SortedRanges)
+	{
+		printf("\t%d:%d - %d:%c:%d:%d\n", o.SourceCB, o.SourceOffset, o.DestCBIndex, o.DestCBPrecision, o.DestOffset, o.Size);
+	}
+}
+
+// Returns true if the FramebufferFetchES2() 'intrinsic' is used
+static bool UsesFramebufferFetchES2(exec_list* Instructions)
+{
+	struct SFindFramebufferFetchES2Intrinsic : public ir_hierarchical_visitor
+	{
+		bool bFound;
+		SFindFramebufferFetchES2Intrinsic() : bFound(false) {}
+
+		virtual ir_visitor_status visit_enter(ir_call* IR) override
+		{
+			if (IR->use_builtin && !strcmp(IR->callee_name(), FRAMEBUFFER_FETCH_ES2))
+			{
+				bFound = true;
+				return visit_stop;
+			}
+
+			return visit_continue;
+		}
+	};
+
+	SFindFramebufferFetchES2Intrinsic Visitor;
+	Visitor.run(Instructions);
+	return Visitor.bFound;
+}
+
+/**
+ * IR visitor used to generate GLSL. Based on ir_print_visitor.
+ */
+class ir_gen_glsl_visitor : public ir_visitor
+{
+	/** Track which multi-dimensional arrays are used. */
+	struct md_array_entry : public exec_node
+	{
+		const glsl_type* type;
+	};
+
+	/** Track external variables. */
+	struct extern_var : public exec_node
+	{
+		ir_variable* var;
+		explicit extern_var(ir_variable* in_var) : var(in_var) {}
+	};
+
+	/** External variables. */
+	exec_list input_variables;
+	exec_list output_variables;
+	exec_list uniform_variables;
+	exec_list sampler_variables;
+	exec_list image_variables;
+
+	/** Data tied globally to the shader via attributes */
+	bool early_depth_stencil;
+	int wg_size_x;
+	int wg_size_y;
+	int wg_size_z;
+
+	glsl_tessellation_info tessellation;
+
+	/** Track global instructions. */
+	struct global_ir : public exec_node
+	{
+		ir_instruction* ir;
+		explicit global_ir(ir_instruction* in_ir) : ir(in_ir) {}
+	};
+
+	/** Global instructions. */
+	exec_list global_instructions;
+
+	/** A mapping from ir_variable * -> unique printable names. */
+	hash_table *printable_names;
+	/** Structures required by the code. */
+	hash_table *used_structures;
+	/** Uniform block variables required by the code. */
+	hash_table *used_uniform_blocks;
+	/** Multi-dimensional arrays required by the code. */
+	exec_list used_md_arrays;
+
+	// Code generation flags
+	bool bIsES;
+	_mesa_glsl_parser_targets ShaderTarget;
+
+	bool bGenerateLayoutLocations;
+	
+
+	/** Memory context within which to make allocations. */
+	void *mem_ctx;
+	/** Buffer to which GLSL source is being generated. */
+	char** buffer;
+	/** Indentation level. */
+	int indentation;
+	/** Scope depth. */
+	int scope_depth;
+	/** The number of temporary variables declared in the current scope. */
+	int temp_id;
+	/** The number of global variables declared. */
+	int global_id;
+	/** Whether a semicolon must be printed before the next EOL. */
+	bool needs_semicolon;
+	/**
+	 * Whether uint literals should be printed as int literals. This is a hack
+	 * because glCompileShader crashes on Mac OS X with code like this:
+	 * foo = bar[0u];
+	 */
+	bool should_print_uint_literals_as_ints;
+	/** number of loops in the generated code */
+	int loop_count;
+
+	/** Whether the shader being cross compiled needs EXT_shader_texture_lod. */
+	bool bUsesES2TextureLODExtension;
+
+	// Found dFdx or dFdy
+	bool bUsesDXDY;
+
+	/**
+	 * Return true if the type is a multi-dimensional array. Also, track the
+	 * array.
+	 */
+	bool is_md_array(const glsl_type* type)
+	{
+		if (type->base_type == GLSL_TYPE_ARRAY &&
+			type->fields.array->base_type == GLSL_TYPE_ARRAY)
+		{
+			foreach_iter(exec_list_iterator, iter, used_md_arrays) 
+			{
+				md_array_entry* entry = (md_array_entry*)iter.get();
+				if (entry->type == type)
+					return true;
+			}
+			md_array_entry* entry = new(mem_ctx) md_array_entry();
+			entry->type = type;
+			used_md_arrays.push_tail(entry);
+			return true;
+		}
+		return false;
+	}
+
+	/**
+	 * Fetch/generate a unique name for ir_variable.
+	 *
+	 * GLSL IR permits multiple ir_variables to share the same name.  This works
+	 * fine until we try to print it, when we really need a unique one.
+	 */
+	const char *unique_name(ir_variable *var)
+	{
+		if (var->mode == ir_var_temporary || var->mode == ir_var_auto)
+		{
+			/* Do we already have a name for this variable? */
+			const char *name = (const char *) hash_table_find(this->printable_names, var);
+			if (name == NULL)
+			{
+				bool is_global = (scope_depth == 0 && var->mode != ir_var_temporary);
+				const char *prefix = is_global ? "g" : "t";
+				int var_id = is_global ? global_id++ : temp_id++;
+				name = ralloc_asprintf(mem_ctx, "%s%d", prefix, var_id);
+				hash_table_insert(this->printable_names, (void *)name, var);
+			}
+			return name;
+		}
+
+		/* If there's no conflict, just use the original name */
+		return var->name;
+	}
+
+	/**
+	 * Add tabs/spaces for the current indentation level.
+	 */
+	void indent(void)
+	{
+		for (int i = 0; i < indentation; i++)
+		{
+			ralloc_asprintf_append(buffer, "\t");
+		}
+	}
+
+	/**
+	 * Print out the internal name for a multi-dimensional array.
+	 */
+	void print_md_array_type(const glsl_type *t)
+	{
+		if (t->base_type == GLSL_TYPE_ARRAY)
+		{
+			ralloc_asprintf_append(buffer, "_mdarr_");
+			do 
+			{
+				ralloc_asprintf_append(buffer, "%d_", t->length);
+				t = t->fields.array;
+			} while (t->base_type == GLSL_TYPE_ARRAY);
+			print_base_type(t);
+		}
+	}
+
+	/**
+	 * Print the base type, e.g. vec3.
+	 */
+	void print_base_type(const glsl_type *t)
+	{
+		if (t->base_type == GLSL_TYPE_ARRAY)
+		{
+			print_base_type(t->fields.array);
+		}
+		else if (t->base_type == GLSL_TYPE_INPUTPATCH)
+		{
+			ralloc_asprintf_append(buffer, "/* %s */ ", t->name);
+			print_base_type(t->inner_type);
+		}
+		else if (t->base_type == GLSL_TYPE_OUTPUTPATCH)
+		{
+			ralloc_asprintf_append(buffer, "/* %s */ ", t->name);
+			print_base_type(t->inner_type);
+		}
+		else if ((t->base_type == GLSL_TYPE_STRUCT)
+				&& (strncmp("gl_", t->name, 3) != 0))
+		{
+			ralloc_asprintf_append(buffer, "%s", t->name);
+		}
+		else 
+		{
+			std::string Name = FixHlslName(t);
+			ralloc_asprintf_append(buffer, "%s", Name.c_str());
+		}
+	}
+
+	/**
+	 * Print the portion of the type that appears before a variable declaration.
+	 */
+	void print_type_pre(const glsl_type *t)
+	{
+		if (is_md_array(t))
+		{
+			print_md_array_type(t);
+		}
+		else
+		{
+			print_base_type(t);
+		}
+	}
+
+	/**
+	 * Print the portion of the type that appears after a variable declaration.
+	 */
+	void print_type_post(const glsl_type *t, bool is_unsized = false)
+	{
+		if (t->base_type == GLSL_TYPE_ARRAY && !is_md_array(t))
+		{
+			if (is_unsized)
+			{
+				ralloc_asprintf_append(buffer, "[]");
+				
+			}
+			else
+			{
+				ralloc_asprintf_append(buffer, "[%u]", t->length);
+			}
+		}
+		else if (t->base_type == GLSL_TYPE_INPUTPATCH || t->base_type == GLSL_TYPE_OUTPUTPATCH)
+		{
+			ralloc_asprintf_append(buffer, "[%u] /* %s */", t->patch_length, t->name);
+		}
+	}
+
+	/**
+	 * Print a full variable declaration.
+	 */
+	void print_type_full(const glsl_type *t)
+	{
+		print_type_pre(t);
+		print_type_post(t);
+	}
+
+	/**
+	 * Visit a single instruction. Appends a semicolon and EOL if needed.
+	 */
+	void do_visit(ir_instruction* ir)
+	{
+		needs_semicolon = true;
+		ir->accept(this);
+		if (needs_semicolon)
+		{
+			ralloc_asprintf_append(buffer, ";\n");
+		}
+	}
+
+	/**
+	* \name Visit methods
+	*
+	* As typical for the visitor pattern, there must be one \c visit method for
+	* each concrete subclass of \c ir_instruction.  Virtual base classes within
+	* the hierarchy should not have \c visit methods.
+	*/
+
+	virtual void visit(ir_rvalue *rvalue)
+	{
+		check(0 && "ir_rvalue not handled for GLSL export.");
+	}
+
+	virtual void visit(ir_variable *var)
+	{
+		const char * const centroid_str[] = { "", "centroid " };
+		const char * const invariant_str[] = { "", "invariant " };
+		const char * const patch_constant_str[] = { "", "patch " };
+		const char * const GLSLmode_str[] = { "", "uniform ", "in ", "out ", "inout ", "in ", "", "shared ", "", "", "uniform_ref "};
+		const char * const ESVSmode_str[] = { "", "uniform ", "attribute ", "varying ", "inout ", "in ", "", "shared " };
+		const char * const ESFSmode_str[] = { "", "uniform ", "varying ", "attribute ", "", "in ", "", "shared " };
+		const char * const interp_str[] = { "", "smooth ", "flat ", "noperspective " };
+		const char * const layout_str[] = { "", "layout(origin_upper_left) ", "layout(pixel_center_integer) ", "layout(origin_upper_left,pixel_center_integer) " };
+
+		const char * const * mode_str = bIsES ? ((ShaderTarget == vertex_shader) ? ESVSmode_str : ESFSmode_str) : GLSLmode_str;
+
+		// Check for an initialized const variable
+		// If var is read-only and initialized, set it up as an initialized const
+		bool constInit = false;
+		if (var->has_initializer && var->read_only && (var->constant_initializer || var->constant_value))
+		{
+			ralloc_asprintf_append( buffer, "const ");
+			constInit = true;
+		}
+
+		if (scope_depth == 0)
+		{
+			glsl_base_type base_type = var->type->base_type;
+			if (base_type == GLSL_TYPE_ARRAY)
+			{
+				base_type = var->type->fields.array->base_type;
+			}
+
+			if (var->mode == ir_var_in)
+			{
+				input_variables.push_tail(new(mem_ctx) extern_var(var));
+			}
+			else if (var->mode == ir_var_out)
+			{
+				output_variables.push_tail(new(mem_ctx) extern_var(var));
+			}
+			else if (var->mode == ir_var_uniform && var->type->is_sampler())
+			{
+				sampler_variables.push_tail(new(mem_ctx) extern_var(var));
+			}
+			else if (var->mode == ir_var_uniform && var->type->is_image())
+			{
+				image_variables.push_tail(new(mem_ctx) extern_var(var));
+			}
+			else if (var->mode == ir_var_uniform && base_type == GLSL_TYPE_SAMPLER_STATE)
+			{
+				// ignore sampler state uniforms
+			}
+			else if (var->mode == ir_var_uniform && var->semantic == NULL)
+			{
+				uniform_variables.push_tail(new(mem_ctx) extern_var(var));
+			}
+		}
+
+		if (var->name && strncmp(var->name, "gl_", 3) == 0 &&
+			var->centroid == 0 && var->interpolation == 0 &&
+			var->invariant == 0 && var->origin_upper_left == 0 &&
+			var->pixel_center_integer == 0)
+		{
+			// Don't emit builtin GL variable declarations.
+			needs_semicolon = false;
+		}
+		else if (scope_depth == 0 && var->mode == ir_var_temporary)
+		{
+			global_instructions.push_tail(new(mem_ctx) global_ir(var));
+			needs_semicolon = false;
+		}
+		else
+		{
+			int layout_bits =
+				(var->origin_upper_left ? 0x1 : 0) |
+				(var->pixel_center_integer ? 0x2 : 0);
+
+			if (scope_depth == 0 &&
+			   ((var->mode == ir_var_in) || (var->mode == ir_var_out)) && 
+			   var->is_interface_block)
+			{
+				/**
+				Hack to display our fake structs as what they are supposed to be - interface blocks
+
+				'in' or 'out' variable qualifier becomes interface block declaration start,
+				structure name becomes block name,
+				we add information about block contents, taking type from sole struct member type, and
+				struct variable name becomes block instance name.
+
+				Note: With tessellation, matching interfaces between shaders is tricky, so we need
+				to assign explicit locations to shader input and output variables.
+
+				The reason we use a struct instead of an interface block is that with
+				GL4.2/GL_ARB_separate_shader_objects, you can add a layout(location=foo) to a variable
+				that is not part of an interface block. However, in order to add a location to a variable
+				inside an interface block, you need GL4.4/GL_enhanced_layouts. Since for now, we don't want
+				that dependency, we use structs.
+
+				*/
+
+				ralloc_asprintf_append(
+					buffer,
+					"%s%s%s%s%s",
+					layout_str[layout_bits],
+					centroid_str[var->centroid],
+					invariant_str[var->invariant],
+					patch_constant_str[var->is_patch_constant],
+					mode_str[var->mode]
+					);
+
+				if (bGenerateLayoutLocations && var->explicit_location)
+				{
+					check(layout_bits == 0);
+
+					ralloc_asprintf_append(
+						buffer,
+						"layout(location=%d) %s struct ",
+						var->location,
+						interp_str[var->interpolation]
+						);
+				}
+
+				print_type_pre(var->type);
+
+				const glsl_type* inner_type = var->type;
+				if (inner_type->is_array())
+				{
+					inner_type = inner_type->fields.array;
+				}
+				check(inner_type->is_record());
+				check(inner_type->length==1);
+				const glsl_struct_field* field = &inner_type->fields.structure[0];
+				check(strcmp(field->name,"Data")==0);
+
+				if (bGenerateLayoutLocations && var->explicit_location)
+				{
+					ralloc_asprintf_append(buffer, " { ");
+				}
+				else
+				{
+					ralloc_asprintf_append(buffer, " { %s", interp_str[var->interpolation]);
+				}
+
+				print_type_pre(field->type);
+				ralloc_asprintf_append(buffer, " Data");
+				print_type_post(field->type);
+				ralloc_asprintf_append(buffer, "; }");
+			}
+			else if (var->type->is_image())
+			{
+				const char * const coherent_str[] = { "", "coherent " };
+				const char * const writeonly_str[] = { "", "writeonly " };
+				const char * const type_str[] = { "32ui", "32i", "16f", "32f" };
+				const int writeonly = var->image_write && !(var->image_read);
+
+				check( var->type->inner_type->base_type >= GLSL_TYPE_UINT &&
+						var->type->inner_type->base_type <= GLSL_TYPE_FLOAT );
+
+				ralloc_asprintf_append(
+					buffer,
+					"%s%s%s%s",
+					invariant_str[var->invariant],
+					mode_str[var->mode],
+					coherent_str[var->coherent],
+					writeonly_str[writeonly]
+					);
+
+				if (bGenerateLayoutLocations && var->explicit_location)
+				{
+					//should check here on base type
+					ralloc_asprintf_append(
+						buffer,
+						"layout(r%s,binding=%d) ",
+						type_str[var->type->inner_type->base_type],
+						var->location
+						);
+				}
+				else
+				{
+					//should check here on base type
+					ralloc_asprintf_append(
+						buffer,
+						"layout(r%s) ",
+						type_str[var->type->inner_type->base_type]
+						);
+				}
+
+				if (GDefaultPrecisionIsHalf && var->type->inner_type->base_type == GLSL_TYPE_FLOAT)
+				{
+					ralloc_asprintf_append(buffer, "highp ");
+				}
+				else if (!GDefaultPrecisionIsHalf && var->type->inner_type->base_type == GLSL_TYPE_HALF)
+				{
+					ralloc_asprintf_append(buffer, "mediump ");
+				}
+				else if (var->type->inner_type->is_integer())
+				{
+					ralloc_asprintf_append(buffer, "highp ");
+				}
+				else
+				{
+					// we missed a type
+					check(false);
+				}
+				print_type_pre(var->type);
+			}
+			else
+			{
+
+				char* layout = nullptr;
+
+				if (bGenerateLayoutLocations && var->explicit_location)
+				{
+					check(layout_bits == 0);
+					layout = ralloc_asprintf(nullptr, "layout(location=%d) ", var->location);
+				}
+				
+				ralloc_asprintf_append(
+					buffer,
+					"%s%s%s%s%s%s",
+					layout ? layout : layout_str[layout_bits],
+					centroid_str[var->centroid],
+					invariant_str[var->invariant],
+					patch_constant_str[var->is_patch_constant],
+					mode_str[var->mode],
+					interp_str[var->interpolation]
+					);
+				if (var->type->is_sampler())
+				{
+					if (GDefaultPrecisionIsHalf && var->type->inner_type->base_type == GLSL_TYPE_FLOAT)
+					{
+						ralloc_asprintf_append(buffer, "highp ");
+					}
+					else if (!GDefaultPrecisionIsHalf && var->type->inner_type->base_type == GLSL_TYPE_HALF)
+					{
+						ralloc_asprintf_append(buffer, "mediump ");
+					}
+					else // shadow samplers, integer textures etc
+					{
+						ralloc_asprintf_append(buffer, "highp ");
+					}
+				}
+				else if (GDefaultPrecisionIsHalf && (var->type->base_type == GLSL_TYPE_FLOAT || (var->type->is_array() && var->type->element_type()->base_type == GLSL_TYPE_FLOAT)))
+				{
+					ralloc_asprintf_append(buffer, "highp ");
+				}
+				else if (!GDefaultPrecisionIsHalf && (var->type->base_type == GLSL_TYPE_HALF || (var->type->is_array() && var->type->element_type()->base_type == GLSL_TYPE_HALF)))
+				{
+					ralloc_asprintf_append(buffer, "mediump ");
+				}
+				else if (var->type->is_integer())
+				{
+					ralloc_asprintf_append(buffer, "highp ");
+				}
+
+				if (bGenerateLayoutLocations && var->explicit_location)
+				{
+					ralloc_free(layout);
+				}
+
+				print_type_pre(var->type);
+			}
+			ralloc_asprintf_append(buffer, " %s", unique_name(var));
+			const bool bUnsizedArray = var->mode == ir_var_in && ((ShaderTarget == tessellation_evaluation_shader) || (ShaderTarget == tessellation_control_shader));
+			print_type_post(var->type, bUnsizedArray );
+		}
+
+		// Add the initializer if we need it
+		if (constInit)
+		{
+			ralloc_asprintf_append(buffer, " = ");
+			if (var->constant_initializer)
+			{
+				var->constant_initializer->accept(this);
+			}
+			else
+			{
+				var->constant_value->accept(this);
+			}
+		}
+
+		// add type to used_structures so we can later declare them at the start of the GLSL shader
+		// this is for the case of a variable that is declared, but not later dereferenced (which can happen
+		// when debugging HLSLCC and running without optimization
+		AddTypeToUsedStructs(var->type);
+
+	}
+
+	virtual void visit(ir_function_signature *sig)
+	{
+		// Reset temporary id count.
+		temp_id = 0;
+		bool bPrintComma = false;
+		scope_depth++;
+
+		print_type_full(sig->return_type);
+		ralloc_asprintf_append(buffer, " %s(", sig->function_name());
+
+		foreach_iter(exec_list_iterator, iter, sig->parameters)
+		{
+			ir_variable *const inst = (ir_variable *) iter.get();
+			if (bPrintComma)
+			{
+				ralloc_asprintf_append(buffer, ",");
+			}
+			inst->accept(this);
+			bPrintComma = true;
+		}
+		ralloc_asprintf_append(buffer, ")\n");
+
+		indent();
+		ralloc_asprintf_append(buffer, "{\n");
+
+		if (sig->is_main && !global_instructions.is_empty())
+		{
+			indentation++;
+			foreach_iter(exec_list_iterator, iter, global_instructions)
+			{
+				global_ir* gir = (global_ir*)iter.get();
+				indent();
+				do_visit(gir->ir);
+			}
+			indentation--;
+		}
+
+		//grab the global attributes
+		if (sig->is_main)
+		{
+			early_depth_stencil = sig->is_early_depth_stencil;
+			wg_size_x = sig->wg_size_x;
+			wg_size_y = sig->wg_size_y;
+			wg_size_z = sig->wg_size_z;
+
+			tessellation = sig->tessellation;
+		}
+
+		indentation++;
+		foreach_iter(exec_list_iterator, iter, sig->body)
+		{
+			ir_instruction *const inst = (ir_instruction *) iter.get();
+			indent();
+			do_visit(inst);
+		}
+		indentation--;
+		indent();
+		ralloc_asprintf_append(buffer, "}\n");
+		needs_semicolon = false;
+		scope_depth--;
+	}
+
+	virtual void visit(ir_function *func)
+	{
+		foreach_iter(exec_list_iterator, iter, *func)
+		{
+			ir_function_signature *const sig = (ir_function_signature *) iter.get();
+			if (sig->is_defined && !sig->is_builtin)
+			{
+				indent();
+				sig->accept(this);
+			}
+		}
+		needs_semicolon = false;
+	}
+
+	virtual void visit(ir_expression *expr)
+	{
+		check(scope_depth > 0);
+
+		int numOps = expr->get_num_operands();
+		ir_expression_operation op = expr->operation;
+
+		if (numOps == 1 && op >= ir_unop_first_conversion && op <= ir_unop_last_conversion)
+		{
+			if (op == ir_unop_f2h || op == ir_unop_h2f)
+			{
+				// No need to convert from half<->float as that is part of the precision of a variable
+				expr->operands[0]->accept(this);
+			}
+			else
+			{
+				ralloc_asprintf_append(buffer, "%s(", FixHlslName(expr->type).c_str());
+				expr->operands[0]->accept(this);
+				ralloc_asprintf_append(buffer, ")");
+			}
+		}
+		else if (expr->type->is_scalar() &&
+			((numOps == 1 && op == ir_unop_logic_not) ||
+			 (numOps == 2 && op >= ir_binop_first_comparison && op <= ir_binop_last_comparison) ||
+			 (numOps == 2 && op >= ir_binop_first_logic && op <= ir_binop_last_logic)))
+		{
+			const char* op_str = GLSLExpressionTable[op][3];
+			ralloc_asprintf_append(buffer, "%s(", (numOps == 1) ? op_str : "");
+			expr->operands[0]->accept(this);
+			if (numOps == 2)
+			{
+				ralloc_asprintf_append(buffer, "%s", op_str);
+				expr->operands[1]->accept(this);
+			}
+			ralloc_asprintf_append(buffer, ")");
+		}
+		else if (expr->type->is_vector() && numOps == 2 &&
+			op >= ir_binop_first_logic && op <= ir_binop_last_logic)
+		{
+			ralloc_asprintf_append(buffer, GLSLExpressionTable[op][0], expr->type->vector_elements, expr->type->vector_elements);
+			expr->operands[0]->accept(this);
+			ralloc_asprintf_append(buffer, GLSLExpressionTable[op][1], expr->type->vector_elements);
+			expr->operands[1]->accept(this);
+			ralloc_asprintf_append(buffer, GLSLExpressionTable[op][2]);
+		}
+		else if (op == ir_binop_mod && !expr->type->is_float())
+		{
+			ralloc_asprintf_append(buffer, "((");
+			expr->operands[0]->accept(this);
+			ralloc_asprintf_append(buffer, ")%%(");
+			expr->operands[1]->accept(this);
+			ralloc_asprintf_append(buffer, "))");
+		}
+		else if (op == ir_binop_mul && expr->type->is_matrix()
+			&& expr->operands[0]->type->is_matrix()
+			&& expr->operands[1]->type->is_matrix())
+		{
+			ralloc_asprintf_append(buffer, "matrixCompMult(");
+			expr->operands[0]->accept(this);
+			ralloc_asprintf_append(buffer, ",");
+			expr->operands[1]->accept(this);
+			ralloc_asprintf_append(buffer, ")");
+		}
+		else if (numOps < 4)
+		{
+			if (op == ir_unop_dFdx || op == ir_unop_dFdy)
+			{
+				bUsesDXDY = true;
+			}
+
+			ralloc_asprintf_append(buffer, GLSLExpressionTable[op][0]);
+			for (int i = 0; i < numOps; ++i)
+			{
+				expr->operands[i]->accept(this);
+				ralloc_asprintf_append(buffer, GLSLExpressionTable[op][i+1]);
+			}
+		}
+	}
+
+	virtual void visit(ir_texture *tex)
+	{
+		check(scope_depth > 0);
+
+		const char * const fetch_str[] = { "texture", "texelFetch" };
+		const char * const Dim[] = { "", "2D", "3D", "Cube", "", "", "" };
+		static const char * const size_str[] = { "", "Size" };
+		static const char * const proj_str[] = { "", "Proj" };
+		static const char * const grad_str[] = { "", "Grad" };
+		static const char * const lod_str[] = { "", "Lod" };
+		static const char * const offset_str[] = { "", "Offset" };
+		static const char * const gather_str[] = { "", "Gather" };
+		static const char * const querymips_str[] = { "", "QueryLevels" };
+		static const char * const EXT_str[] = { "", "EXT" };
+		const bool cube_array = tex->sampler->type->sampler_dimensionality == GLSL_SAMPLER_DIM_CUBE && 
+			tex->sampler->type->sampler_array;
+
+		ir_texture_opcode op = tex->op;
+		if (op == ir_txl && tex->sampler->type->sampler_shadow && tex->sampler->type->sampler_dimensionality == GLSL_SAMPLER_DIM_CUBE)
+		{
+			// This very instruction is missing in OpenGL 3.2, so we need to change the sampling to instruction that exists in order for shader to compile
+			op = ir_tex;
+		}
+
+		bool bEmitEXT = false;
+
+		if (bIsES && op == ir_txl)
+		{
+			// See http://www.khronos.org/registry/gles/extensions/EXT/EXT_shader_texture_lod.txt
+			bUsesES2TextureLODExtension = true;
+			bEmitEXT = true;
+		}
+
+		// Emit texture function and sampler.
+		ralloc_asprintf_append(buffer, "%s%s%s%s%s%s%s%s%s%s(",
+			fetch_str[op == ir_txf],
+			bIsES ? Dim[tex->sampler->type->sampler_dimensionality] : "",
+			gather_str[op == ir_txg],
+			size_str[op == ir_txs],
+			querymips_str[op == ir_txm],
+			proj_str[tex->projector != 0],
+			grad_str[op == ir_txd],
+			lod_str[op == ir_txl],
+			offset_str[tex->offset != 0],
+			EXT_str[(int)bEmitEXT]
+		);
+		tex->sampler->accept(this);
+
+		// Emit coordinates.
+		if ( (op == ir_txs && tex->lod_info.lod) || op == ir_txm)
+		{
+			if (!tex->sampler->type->sampler_ms && op != ir_txm)
+			{
+				ralloc_asprintf_append(buffer, ",");
+				tex->lod_info.lod->accept(this);
+			}
+		}
+		else if (tex->sampler->type->sampler_shadow && (op != ir_txg && !cube_array))
+		{
+			int coord_dims = 0;
+			switch (tex->sampler->type->sampler_dimensionality)
+			{
+				case GLSL_SAMPLER_DIM_1D: coord_dims = 2; break;
+				case GLSL_SAMPLER_DIM_2D: coord_dims = 3; break;
+				case GLSL_SAMPLER_DIM_3D: coord_dims = 4; break;
+				case GLSL_SAMPLER_DIM_CUBE: coord_dims = 4; break;
+				default: check(0 && "Shadow sampler has unsupported dimensionality.");
+			}
+			ralloc_asprintf_append(buffer, ",vec%d(", coord_dims);
+			tex->coordinate->accept(this);
+			ralloc_asprintf_append(buffer, ",");
+			tex->shadow_comparitor->accept(this);
+			ralloc_asprintf_append(buffer, ")");
+		}
+		else
+		{
+			ralloc_asprintf_append(buffer, ",");
+			tex->coordinate->accept(this);
+		}
+
+		// Emit gather compare value
+		if (tex->sampler->type->sampler_shadow && (op == ir_txg || cube_array))
+		{
+			ralloc_asprintf_append(buffer, ",");
+			tex->shadow_comparitor->accept(this);
+		}
+
+		// Emit sample index.
+		if (op == ir_txf && tex->sampler->type->sampler_ms)
+		{
+			ralloc_asprintf_append(buffer, ",");
+			tex->lod_info.sample_index->accept(this);
+		}
+
+		// Emit LOD.
+		if (op == ir_txl ||
+		   (op == ir_txf && tex->lod_info.lod &&
+		   !tex->sampler->type->sampler_ms && !tex->sampler->type->sampler_buffer))
+		{
+			ralloc_asprintf_append(buffer, ",");
+			tex->lod_info.lod->accept(this);
+		}
+
+		// Emit gradients.
+		if (op == ir_txd)
+		{
+			ralloc_asprintf_append(buffer, ",");
+			tex->lod_info.grad.dPdx->accept(this);
+			ralloc_asprintf_append(buffer, ",");
+			tex->lod_info.grad.dPdy->accept(this);
+		}
+		else if (op == ir_txb)
+		{
+			ralloc_asprintf_append(buffer, ",");
+			tex->lod_info.bias->accept(this);
+		}
+
+		// Emit offset.
+		if (tex->offset)
+		{
+			ralloc_asprintf_append(buffer, ",");
+			tex->offset->accept(this);
+		}
+
+		// Emit channel selection for gather
+		if (op == ir_txg && tex->channel > ir_channel_none)
+		{
+			check( tex->channel < ir_channel_unknown);
+			ralloc_asprintf_append(buffer, ", %d", int(tex->channel) - 1);
+		}
+
+		ralloc_asprintf_append(buffer, ")");
+	}
+
+	virtual void visit(ir_swizzle *swizzle)
+	{
+		check(scope_depth > 0);
+
+		const unsigned mask[4] =
+		{
+			swizzle->mask.x,
+			swizzle->mask.y,
+			swizzle->mask.z,
+			swizzle->mask.w,
+		};
+
+		if (swizzle->val->type->is_scalar())
+		{
+			// Scalar -> Vector swizzles must use the constructor syntax.
+			if (swizzle->type->is_scalar() == false)
+			{
+				print_type_full(swizzle->type);
+				ralloc_asprintf_append(buffer, "(");
+				swizzle->val->accept(this);
+				ralloc_asprintf_append(buffer, ")");
+			}
+		}
+		else
+		{
+			const bool is_constant = swizzle->val->as_constant() != nullptr;
+			if (is_constant)
+			{
+				ralloc_asprintf_append(buffer, "(");
+			}
+			swizzle->val->accept(this);
+			if (is_constant)
+			{
+				ralloc_asprintf_append(buffer, ")");
+			}
+			ralloc_asprintf_append(buffer, ".");
+			for (unsigned i = 0; i < swizzle->mask.num_components; i++)
+			{
+				ralloc_asprintf_append(buffer, "%c", "xyzw"[mask[i]]);
+			}
+		}
+	}
+
+	virtual void visit(ir_dereference_variable *deref)
+	{
+		check(scope_depth > 0);
+
+		ir_variable* var = deref->variable_referenced();
+
+		ralloc_asprintf_append(buffer, unique_name(var));
+
+
+		// add type to used_structures so we can later declare them at the start of the GLSL shader
+		AddTypeToUsedStructs(var->type);
+
+
+		if (var->mode == ir_var_uniform && var->semantic != NULL)
+		{
+			if (hash_table_find(used_uniform_blocks, var->semantic) == NULL)
+			{
+				hash_table_insert(used_uniform_blocks, (void*)var->semantic, var->semantic);
+			}
+		}
+
+		if (is_md_array(deref->type))
+		{
+			ralloc_asprintf_append(buffer, ".Inner");
+		}
+	}
+
+	virtual void visit(ir_dereference_array *deref)
+	{
+		check(scope_depth > 0);
+
+		deref->array->accept(this);
+
+		// Make extra sure crappy Mac OS X compiler won't have any reason to crash
+		bool enforceInt = false;
+
+		if (deref->array_index->type->base_type == GLSL_TYPE_UINT)
+		{
+			if (deref->array_index->ir_type == ir_type_constant)
+			{
+				should_print_uint_literals_as_ints = true;
+			}
+			else
+			{
+				enforceInt = true;
+			}
+		}
+
+		if (enforceInt)
+		{
+			ralloc_asprintf_append(buffer, "[int(");
+		}
+		else
+		{
+			ralloc_asprintf_append(buffer, "[");
+		}
+
+		deref->array_index->accept(this);
+		should_print_uint_literals_as_ints = false;
+
+		if (enforceInt)
+		{
+			ralloc_asprintf_append(buffer, ")]");
+		}
+		else
+		{
+			ralloc_asprintf_append(buffer, "]");
+		}
+
+		if (is_md_array(deref->array->type))
+		{
+			ralloc_asprintf_append(buffer, ".Inner");
+		}
+	}
+
+	void print_image_op( ir_dereference_image *deref, ir_rvalue *src)
+	{
+		const char* swizzle[] =
+		{
+			"x", "xy", "xyz", "xyzw"
+		};
+		const char* expand[] =
+		{
+			"xxxx", "xyxx", "xyzx", "xyzw"
+		};
+		const char* int_cast[] =
+		{
+			"int", "ivec2", "ivec3", "ivec4"
+		};
+		const int dst_elements = deref->type->vector_elements;
+		const int src_elements = (src) ? src->type->vector_elements : 1;
+		
+		check( 1 <= dst_elements && dst_elements <= 4);
+		check( 1 <= src_elements && src_elements <= 4);
+
+		if ( deref->op == ir_image_access)
+		{
+			if ( src == NULL )
+			{
+				ralloc_asprintf_append( buffer, "imageLoad( " );
+				deref->image->accept(this);
+				ralloc_asprintf_append(buffer, ", ");
+				deref->image_index->accept(this);
+				ralloc_asprintf_append(buffer, ").%s", swizzle[dst_elements-1]);
+			}
+			else
+			{
+				ralloc_asprintf_append( buffer, "imageStore( " );
+				deref->image->accept(this);
+				ralloc_asprintf_append(buffer, ", ");
+				deref->image_index->accept(this);
+				ralloc_asprintf_append(buffer, ", ");
+				src->accept(this);
+				ralloc_asprintf_append(buffer, ".%s)", expand[src_elements-1]);
+			}
+		}
+		else if ( deref->op == ir_image_dimensions)
+		{
+			ralloc_asprintf_append( buffer, "imageSize( " );
+			deref->image->accept(this);
+			ralloc_asprintf_append(buffer, ")");
+		}
+		else
+		{
+			check( !"Unknown image operation");
+		}
+	}
+
+	virtual void visit(ir_dereference_image *deref)
+	{
+		check(scope_depth > 0);
+
+		print_image_op( deref, NULL);
+		
+	}
+
+	virtual void visit(ir_dereference_record *deref)
+	{
+		check(scope_depth > 0);
+
+		deref->record->accept(this);
+		ralloc_asprintf_append(buffer, ".%s", deref->field);
+
+		if (is_md_array(deref->type))
+		{
+			ralloc_asprintf_append(buffer, ".Inner");
+		}
+	}
+
+	virtual void visit(ir_assignment *assign)
+	{
+		if (scope_depth == 0)
+		{
+			global_instructions.push_tail(new(mem_ctx) global_ir(assign));
+			needs_semicolon = false;
+			return;
+		}
+
+		// constant variables with initializers are statically assigned
+		ir_variable *var = assign->lhs->variable_referenced();
+		if (var->has_initializer && var->read_only && (var->constant_initializer || var->constant_value))
+		{
+			//This will leave a blank line with a semi-colon
+			return;
+		}
+
+		if (assign->condition)
+		{
+			ralloc_asprintf_append(buffer, "if(");
+			assign->condition->accept(this);
+			ralloc_asprintf_append(buffer, ") { ");
+		}
+
+		if (assign->lhs->as_dereference_image() != NULL)
+		{
+			/** EHart - should the write mask be checked here? */
+			print_image_op( assign->lhs->as_dereference_image(), assign->rhs);
+		}
+		else
+		{
+			char mask[6];
+			unsigned j = 1;
+			if (assign->lhs->type->is_scalar() == false ||
+				assign->write_mask != 0x1)
+			{
+				for (unsigned i = 0; i < 4; i++)
+				{
+					if ((assign->write_mask & (1 << i)) != 0)
+					{
+						mask[j] = "xyzw"[i];
+						j++;
+					}
+				}
+			}
+			mask[j] = '\0';
+
+			mask[0] = (j == 1) ? '\0' : '.';
+
+			assign->lhs->accept(this);
+			ralloc_asprintf_append(buffer, "%s = ", mask);
+			assign->rhs->accept(this);
+		}
+
+		if (assign->condition)
+		{
+			ralloc_asprintf_append(buffer, "%s }", needs_semicolon ? ";" : "");
+		}
+	}
+
+	void print_constant(ir_constant *constant, int index)
+	{
+		if (constant->type->is_float())
+		{
+			if (constant->is_component_finite(index))
+			{
+				float value = constant->value.f[index];
+				// Original formatting code relied on %f style formatting
+				// %e is more accureate, and has been available since at least ES 2.0
+				// leaving original code in place, in case some drivers don't properly handle it
+#if 0
+				const char *format = (fabsf(fmodf(value,1.0f)) < 1.e-8f) ? "%.1f" : "%.8f";
+#else
+				const char *format = "%e";
+#endif
+				ralloc_asprintf_append(buffer, format, value);
+			}
+			else
+			{
+				switch (constant->value.u[index])
+				{
+					case 0x7f800000u:
+						ralloc_asprintf_append(buffer, "(1.0/0.0)");
+						break;
+
+					case 0xffc00000u:
+						ralloc_asprintf_append(buffer, "(0.0/0.0)");
+						break;
+
+					case 0xff800000u:
+						ralloc_asprintf_append(buffer, "(-1.0/0.0)");
+						break;
+
+					default:
+						check(0);
+				}
+			}
+		}
+		else if (constant->type->base_type == GLSL_TYPE_INT)
+		{
+			ralloc_asprintf_append(buffer, "%d", constant->value.i[index]);
+		}
+		else if (constant->type->base_type == GLSL_TYPE_UINT)
+		{
+			ralloc_asprintf_append(buffer, "%u%s",
+				constant->value.u[index],
+				should_print_uint_literals_as_ints ? "" : "u"
+				);
+		}
+		else if (constant->type->base_type == GLSL_TYPE_BOOL)
+		{
+			ralloc_asprintf_append(buffer, "%s", constant->value.b[index] ? "true" : "false");
+		}
+	}
+
+	virtual void visit(ir_constant *constant)
+	{
+		if (constant->type == glsl_type::float_type
+			|| constant->type == glsl_type::half_type
+			|| constant->type == glsl_type::bool_type
+			|| constant->type == glsl_type::int_type
+			|| constant->type == glsl_type::uint_type)
+		{
+			print_constant(constant, 0);
+		}
+		else if (constant->type->is_record())
+		{
+			print_type_full(constant->type);
+			ralloc_asprintf_append(buffer, "(");
+			ir_constant* value = (ir_constant*)constant->components.get_head();
+			if (value)
+			{
+				value->accept(this);
+			}
+			for (uint32 i = 1; i < constant->type->length; i++)
+			{
+				check(value);
+				value = (ir_constant*)value->next;
+				if (value)
+				{
+					ralloc_asprintf_append(buffer, ",");
+					value->accept(this);
+				}
+			}
+			ralloc_asprintf_append(buffer, ")");
+		}
+		else if (constant->type->is_array())
+		{
+			print_type_full(constant->type);
+			ralloc_asprintf_append(buffer, "(");
+			constant->get_array_element(0)->accept(this);
+			for (uint32 i = 1; i < constant->type->length; ++i)
+			{
+				ralloc_asprintf_append(buffer, ",");
+				constant->get_array_element(i)->accept(this);
+			}
+			ralloc_asprintf_append(buffer, ")");
+		}
+		else
+		{
+			print_type_full(constant->type);
+			ralloc_asprintf_append(buffer, "(");
+			print_constant(constant, 0);
+			int num_components = constant->type->components();
+			for (int i = 1; i < num_components; ++i)
+			{
+				ralloc_asprintf_append(buffer, ",");
+				print_constant(constant, i);
+			}
+			ralloc_asprintf_append(buffer, ")");
+		}
+	}
+
+	virtual void visit(ir_call *call)
+	{
+		if (scope_depth == 0)
+		{
+			global_instructions.push_tail(new(mem_ctx) global_ir(call));
+			needs_semicolon = false;
+			return;
+		}
+
+		if (call->return_deref)
+		{
+			call->return_deref->accept(this);
+			ralloc_asprintf_append(buffer, " = ");
+		}
+		ralloc_asprintf_append(buffer, "%s(", call->callee_name());
+		bool bPrintComma = false;
+		foreach_iter(exec_list_iterator, iter, *call)
+		{
+			ir_instruction *const inst = (ir_instruction *) iter.get();
+			if (bPrintComma)
+			{
+				ralloc_asprintf_append(buffer, ",");
+			}
+			inst->accept(this);
+			bPrintComma = true;
+		}
+		ralloc_asprintf_append(buffer, ")");
+	}
+
+	virtual void visit(ir_return *ret)
+	{
+		check(scope_depth > 0);
+
+		ralloc_asprintf_append(buffer, "return ");
+		ir_rvalue *const value = ret->get_value();
+		if (value)
+		{
+			value->accept(this);
+		}
+	}
+
+	virtual void visit(ir_discard *discard)
+	{
+		check(scope_depth > 0);
+
+		if (discard->condition)
+		{
+			ralloc_asprintf_append(buffer, "if (");
+			discard->condition->accept(this);
+			ralloc_asprintf_append(buffer, ") ");
+		}
+		ralloc_asprintf_append(buffer, "discard");
+	}
+
+	bool try_conditional_move(ir_if *expr)
+	{
+		ir_dereference_variable *dest_deref = NULL;
+		ir_rvalue *true_value = NULL;
+		ir_rvalue *false_value = NULL;
+		unsigned write_mask = 0;
+		const glsl_type *assign_type = NULL;
+		int num_inst;
+
+		num_inst = 0;
+		foreach_iter(exec_list_iterator, iter, expr->then_instructions)
+		{
+			if (num_inst > 0)
+			{
+				// multiple instructions? not a conditional move
+				return false;
+			}
+
+			ir_instruction *const inst = (ir_instruction *) iter.get();
+			ir_assignment *assignment = inst->as_assignment();
+			if (assignment && (assignment->rhs->ir_type == ir_type_dereference_variable || assignment->rhs->ir_type == ir_type_constant))
+			{
+				dest_deref = assignment->lhs->as_dereference_variable();
+				true_value = assignment->rhs;
+				write_mask = assignment->write_mask;
+			}
+			num_inst++;
+		}
+
+		if (dest_deref == NULL || true_value == NULL)
+			return false;
+
+		num_inst = 0;
+		foreach_iter(exec_list_iterator, iter, expr->else_instructions)
+		{
+			if (num_inst > 0)
+			{
+				// multiple instructions? not a conditional move
+				return false;
+			}
+
+			ir_instruction *const inst = (ir_instruction *) iter.get();
+			ir_assignment *assignment = inst->as_assignment();
+			if (assignment && (assignment->rhs->ir_type == ir_type_dereference_variable || assignment->rhs->ir_type == ir_type_constant))
+			{
+				ir_dereference_variable *tmp_deref = assignment->lhs->as_dereference_variable();
+				if (tmp_deref
+					&& tmp_deref->var == dest_deref->var
+					&& tmp_deref->type == dest_deref->type
+					&& assignment->write_mask == write_mask)
+				{
+					false_value= assignment->rhs;
+				}
+			}
+			num_inst++;
+		}
+
+		if (false_value == NULL)
+			return false;
+
+		char mask[6];
+		unsigned j = 1;
+		if (dest_deref->type->is_scalar() == false || write_mask != 0x1)
+		{
+			for (unsigned i = 0; i < 4; i++)
+			{
+				if ((write_mask & (1 << i)) != 0)
+				{
+					mask[j] = "xyzw"[i];
+					j++;
+				}
+			}
+		}
+		mask[j] = '\0';
+		mask[0] = (j == 1) ? '\0' : '.';
+
+		dest_deref->accept(this);
+		ralloc_asprintf_append(buffer, "%s = (", mask);
+		expr->condition->accept(this);
+		ralloc_asprintf_append(buffer, ")?(");
+		true_value->accept(this);
+		ralloc_asprintf_append(buffer, "):(");
+		false_value->accept(this);
+		ralloc_asprintf_append(buffer, ")");
+
+		return true;
+	}
+
+	virtual void visit(ir_if *expr)
+	{
+		check(scope_depth > 0);
+
+		if (try_conditional_move(expr) == false)
+		{
+			ralloc_asprintf_append(buffer, "if (");
+			expr->condition->accept(this);
+			ralloc_asprintf_append(buffer, ")\n");
+			indent();
+			ralloc_asprintf_append(buffer, "{\n");
+
+			indentation++;
+			foreach_iter(exec_list_iterator, iter, expr->then_instructions)
+			{
+			ir_instruction *const inst = (ir_instruction *) iter.get();
+				indent();
+				do_visit(inst);
+			}
+			indentation--;
+
+			indent();
+			ralloc_asprintf_append(buffer, "}\n");
+
+			if (!expr->else_instructions.is_empty())
+			{
+				indent();
+				ralloc_asprintf_append(buffer, "else\n");
+				indent();
+				ralloc_asprintf_append(buffer, "{\n");
+
+				indentation++;
+				foreach_iter(exec_list_iterator, iter, expr->else_instructions)
+				{
+				ir_instruction *const inst = (ir_instruction *) iter.get();
+					indent();
+					do_visit(inst);
+				}
+				indentation--;
+
+				indent();
+				ralloc_asprintf_append(buffer, "}\n");
+			}
+
+			needs_semicolon = false;
+		}
+	}
+
+	virtual void visit(ir_loop *loop)
+	{
+		check(scope_depth > 0);
+
+		if (loop->counter && loop->to)
+		{
+			// IR cmp operator is when to terminate loop; whereas GLSL for loop syntax
+			// is while to continue the loop. Invert the meaning of operator when outputting.
+			const char* termOp = NULL;
+			switch (loop->cmp)
+			{
+				case ir_binop_less: termOp = ">="; break;
+				case ir_binop_greater: termOp = "<="; break;
+				case ir_binop_lequal: termOp = ">"; break;
+				case ir_binop_gequal: termOp = "<"; break;
+				case ir_binop_equal: termOp = "!="; break;
+				case ir_binop_nequal: termOp = "=="; break;
+				default: check(false);
+			}
+			ralloc_asprintf_append(buffer, "for (;%s%s", unique_name(loop->counter), termOp);
+			loop->to->accept (this);
+			ralloc_asprintf_append(buffer, ";)\n");
+		}
+		else
+		{
+#if 1
+			ralloc_asprintf_append(buffer, "for (;;)\n");
+#else
+			ralloc_asprintf_append(buffer, "for ( int loop%d = 0; loop%d < 256; loop%d ++)\n", loop_count, loop_count, loop_count);
+			loop_count++;
+#endif
+		}
+		indent();
+		ralloc_asprintf_append(buffer, "{\n");
+
+		indentation++;
+		foreach_iter(exec_list_iterator, iter, loop->body_instructions)
+		{
+			ir_instruction *const inst = (ir_instruction *) iter.get();
+			indent();
+			do_visit(inst);
+		}
+		indentation--;
+
+		indent();
+		ralloc_asprintf_append(buffer, "}\n");
+
+		needs_semicolon = false;
+	}
+
+	virtual void visit(ir_loop_jump *jmp)
+	{
+		check(scope_depth > 0);
+
+		ralloc_asprintf_append(buffer, "%s",
+			jmp->is_break() ? "break" : "continue");
+	}
+
+	virtual void visit(ir_atomic *ir)
+	{
+		const char *sharedAtomicFunctions[] = 
+		{
+			"atomicAdd",
+			"atomicAnd",
+			"atomicMin",
+			"atomicMax",
+			"atomicOr",
+			"atomicXor",
+			"atomicExchange",
+			"atomicCompSwap"
+		};
+		const char *imageAtomicFunctions[] = 
+		{
+			"imageAtomicAdd",
+			"imageAtomicAnd",
+			"imageAtomicMin",
+			"imageAtomicMax",
+			"imageAtomicOr",
+			"imageAtomicXor",
+			"imageAtomicExchange",
+			"imageAtomicCompSwap"
+		};
+		check(scope_depth > 0);
+		const bool is_image = ir->memory_ref->as_dereference_image() != NULL;
+
+		ir->lhs->accept(this);
+		if (!is_image)
+		{
+			ralloc_asprintf_append(buffer, " = %s(",
+				sharedAtomicFunctions[ir->operation]);
+			ir->memory_ref->accept(this);
+			ralloc_asprintf_append(buffer, ", ");
+			ir->operands[0]->accept(this);
+			if (ir->operands[1])
+			{
+				ralloc_asprintf_append(buffer, ", ");
+				ir->operands[1]->accept(this);
+			}
+			ralloc_asprintf_append(buffer, ")");
+		}
+		else
+		{
+			ir_dereference_image *image = ir->memory_ref->as_dereference_image();
+			ralloc_asprintf_append(buffer, " = %s(",
+				imageAtomicFunctions[ir->operation]);
+			image->image->accept(this);
+			ralloc_asprintf_append(buffer, ", ");
+			image->image_index->accept(this);
+			ralloc_asprintf_append(buffer, ", ");
+			ir->operands[0]->accept(this);
+			if (ir->operands[1])
+			{
+				ralloc_asprintf_append(buffer, ", ");
+				ir->operands[1]->accept(this);
+			}
+			ralloc_asprintf_append(buffer, ")");
+		}
+	}
+
+	void AddTypeToUsedStructs(const glsl_type* type);
+
+	/**
+	 * Declare structs used to simulate multi-dimensional arrays.
+	 */
+	void declare_md_array_struct(const glsl_type* type, hash_table* ht)
+	{
+		check(type->is_array());
+
+		if (hash_table_find(ht, (void*)type) == NULL)
+		{
+			const glsl_type* subtype = type->fields.array;
+			if (subtype->base_type == GLSL_TYPE_ARRAY)
+			{
+				declare_md_array_struct(subtype, ht);
+
+				ralloc_asprintf_append(buffer, "struct ");
+				print_md_array_type(type);
+				ralloc_asprintf_append(buffer, "\n{\n\t");
+				print_md_array_type(subtype);
+				ralloc_asprintf_append(buffer, " Inner[%u];\n};\n\n", type->length);
+			}
+			else
+			{
+				ralloc_asprintf_append(buffer, "struct ");
+				print_md_array_type(type);
+				ralloc_asprintf_append(buffer, "\n{\n\t");
+				print_type_pre(type);
+				ralloc_asprintf_append(buffer, " Inner");
+				print_type_post(type);
+				ralloc_asprintf_append(buffer, ";\n};\n\n");
+			}
+			hash_table_insert(ht, (void*)type, (void*)type);
+		}
+	}
+
+	/**
+	 * Declare structs used by the code that has been generated.
+	 */
+	void declare_structs(_mesa_glsl_parse_state* state)
+	{
+		// If any variable in a uniform block is in use, the entire uniform block
+		// must be present including structs that are not actually accessed.
+		for (unsigned i = 0; i < state->num_uniform_blocks; i++)
+		{
+			const glsl_uniform_block* block = state->uniform_blocks[i];
+			if (hash_table_find(used_uniform_blocks, block->name))
+			{
+				for (unsigned var_index = 0; var_index < block->num_vars; ++var_index)
+				{
+					const glsl_type* type = block->vars[var_index]->type;
+					if (type->base_type == GLSL_TYPE_STRUCT &&
+						hash_table_find(used_structures, type) == NULL)
+					{
+						hash_table_insert(used_structures, (void*)type, type);
+					}
+				}
+			}
+		}
+
+		// If otherwise unused structure is a member of another, used structure, the unused structure is also, in fact, used
+		{
+			int added_structure_types;
+			do
+			{
+				added_structure_types = 0;
+				for (unsigned i = 0; i < state->num_user_structures; i++)
+				{
+					const glsl_type *const s = state->user_structures[i];
+
+					if (hash_table_find(used_structures, s) == NULL)
+					{
+						continue;
+					}
+
+					for (unsigned j = 0; j < s->length; j++)
+					{
+						const glsl_type* type = s->fields.structure[j].type;
+
+						if (type->base_type == GLSL_TYPE_STRUCT)
+						{
+							if (hash_table_find(used_structures, type) == NULL)
+							{
+								hash_table_insert(used_structures, (void*)type, type);
+								++added_structure_types;
+							}
+						}
+						else if (type->base_type == GLSL_TYPE_ARRAY && type->fields.array->base_type == GLSL_TYPE_STRUCT)
+						{
+							if (hash_table_find(used_structures, type->fields.array) == NULL)
+							{
+								hash_table_insert(used_structures, (void*)type->fields.array, type->fields.array);
+							}
+						}
+						else if ((type->base_type == GLSL_TYPE_INPUTPATCH || type->base_type == GLSL_TYPE_OUTPUTPATCH) && type->inner_type->base_type == GLSL_TYPE_STRUCT)
+						{
+							if (hash_table_find(used_structures, type->inner_type) == NULL)
+							{
+								hash_table_insert(used_structures, (void*)type->inner_type, type->inner_type);
+							}
+						}
+					}
+				}
+			}
+			while( added_structure_types > 0 );
+		}
+
+		// Generate structures that allow support for multi-dimensional arrays.
+		{
+			hash_table* ht = hash_table_ctor(32, hash_table_pointer_hash, hash_table_pointer_compare);
+			foreach_iter(exec_list_iterator, iter, used_md_arrays) 
+			{
+				md_array_entry* entry = (md_array_entry*)iter.get();
+				declare_md_array_struct(entry->type, ht);
+			}
+			hash_table_dtor(ht);
+		}
+
+#ifdef OPTIMIZE_ANON_STRUCTURES_OUT
+		// If a uniform block consists of a single, anonymous structure, don't declare this structure
+		// separately. We'll remove it entirely during uniform block code generation, and name the
+		// uniform block instead.
+		for (unsigned i = 0; i < state->num_uniform_blocks; i++)
+		{
+			const glsl_uniform_block* block = state->uniform_blocks[i];
+			if (hash_table_find(used_uniform_blocks, block->name))
+			{
+				if (block->num_vars == 1)
+				{
+					ir_variable* var = block->vars[0];
+					const glsl_type* type = var->type;
+
+					if (type->base_type == GLSL_TYPE_STRUCT &&
+						type->name &&
+						!strcmp( var->name, block->name) &&
+						!strncmp(type->name, "anon_struct_", 12))
+					{
+						hash_table_remove(used_structures, type);
+					}
+				}
+			}
+		}
+#endif // OPTIMIZE_ANON_STRUCTURES_OUT
+
+		for (unsigned i = 0; i < state->num_user_structures; i++)
+		{
+			const glsl_type *const s = state->user_structures[i];
+
+			if (hash_table_find(used_structures, s) == NULL)
+			{
+				continue;
+			}
+
+			ralloc_asprintf_append(buffer, "struct %s\n{\n", s->name);
+
+			if (s->length == 0)
+			{
+				ralloc_asprintf_append(buffer, "\thighp float glsl_doesnt_like_empty_structs;\n");
+			}
+			else
+			{
+				for (unsigned j = 0; j < s->length; j++)
+				{
+					ralloc_asprintf_append(buffer, "\t%s ", (state->language_version == 310) ? "highp" : "") ;
+					print_type_pre(s->fields.structure[j].type);
+					ralloc_asprintf_append(buffer, " %s", s->fields.structure[j].name);
+					print_type_post(s->fields.structure[j].type);
+					ralloc_asprintf_append(buffer, ";\n");
+				}
+			}
+			ralloc_asprintf_append(buffer, "};\n\n");
+		}
+
+		unsigned num_used_blocks = 0;
+		for (unsigned i = 0; i < state->num_uniform_blocks; i++)
+		{
+			const glsl_uniform_block* block = state->uniform_blocks[i];
+			if (hash_table_find(used_uniform_blocks, block->name))
+			{
+				const char* block_name = block->name;
+				if (state->has_packed_uniforms)
+				{
+					block_name = ralloc_asprintf(mem_ctx, "%sb%u",
+						glsl_variable_tag_from_parser_target(state->target),
+						num_used_blocks
+						);
+				}
+				ralloc_asprintf_append(buffer, "layout(std140) uniform %s\n{\n", block_name);
+
+				bool optimized_structure_out = false;
+
+#ifdef OPTIMIZE_ANON_STRUCTURES_OUT
+				if (block->num_vars == 1)
+				{
+					ir_variable* var = block->vars[0];
+					const glsl_type* type = var->type;
+
+					if (type->base_type == GLSL_TYPE_STRUCT &&
+						type->name &&
+						!strcmp( var->name, block->name) &&
+						!strncmp(type->name, "anon_struct_", 12))
+					{
+						for (unsigned j = 0; j < type->length; j++)
+						{
+							ralloc_asprintf_append(buffer, "\t%s",  (state->language_version == 310) ? "highp" : "");
+							print_type_pre(type->fields.structure[j].type);
+							ralloc_asprintf_append(buffer, " %s", type->fields.structure[j].name);
+							print_type_post(type->fields.structure[j].type);
+							ralloc_asprintf_append(buffer, ";\n");
+						}
+						ralloc_asprintf_append(buffer, "} %s;\n\n", block->name);
+						optimized_structure_out = true;
+					}
+				}
+#endif
+
+				if (!optimized_structure_out)
+				{
+					for (unsigned var_index = 0; var_index < block->num_vars; ++var_index)
+					{
+						ir_variable* var = block->vars[var_index];
+
+						//EHart - name-mangle variables to prevent colliding names
+						ralloc_asprintf_append(buffer, "#define %s %s%s\n", var->name, var->name, block_name);
+
+						ralloc_asprintf_append(buffer, "\t%s", (state->language_version == 310) ? "highp " : "");
+						print_type_pre(var->type);
+						ralloc_asprintf_append(buffer, " %s", var->name);
+						print_type_post(var->type);
+						ralloc_asprintf_append(buffer, ";\n");
+					}
+					ralloc_asprintf_append(buffer, "};\n\n");
+				}
+
+				num_used_blocks++;
+			}
+		}
+	}
+
+	void PrintPackedSamplers(_mesa_glsl_parse_state::TUniformList& Samplers, TStringToSetMap& TextureToSamplerMap)
+	{
+		bool bPrintHeader = true;
+		bool bNeedsComma = false;
+		for (_mesa_glsl_parse_state::TUniformList::iterator Iter = Samplers.begin(); Iter != Samplers.end(); ++Iter)
+		{
+			glsl_packed_uniform& Sampler = *Iter;
+			std::string SamplerStates("");
+			TStringToSetMap::iterator IterFound = TextureToSamplerMap.find(Sampler.Name);
+			if (IterFound != TextureToSamplerMap.end())
+			{
+				TStringSet& ListSamplerStates = IterFound->second;
+				check(!ListSamplerStates.empty());
+				for (TStringSet::iterator IterSS = ListSamplerStates.begin(); IterSS != ListSamplerStates.end(); ++IterSS)
+				{
+					if (IterSS == ListSamplerStates.begin())
+					{
+						SamplerStates += "[";
+					}
+					else
+					{
+						SamplerStates += ",";
+					}
+					SamplerStates += *IterSS;
+				}
+
+				SamplerStates += "]";
+			}
+
+			ralloc_asprintf_append(
+				buffer,
+				"%s%s(%u:%u%s)",
+				bNeedsComma ? "," : "",
+				Sampler.Name.c_str(),
+				Sampler.offset,
+				Sampler.num_components,
+				SamplerStates.c_str()
+				);
+
+			bNeedsComma = true;
+		}
+/*
+		for (TStringToSetMap::iterator Iter = state->TextureToSamplerMap.begin(); Iter != state->TextureToSamplerMap.end(); ++Iter)
+		{
+		const std::string& Texture = Iter->first;
+		TStringSet& Samplers = Iter->second;
+		if (!Samplers.empty())
+		{
+		if (bFirstTexture)
+		{
+		bFirstTexture = false;
+		}
+		else
+		{
+		ralloc_asprintf_append(buffer, ",");
+		}
+
+		ralloc_asprintf_append(buffer, "%s(", Texture.c_str());
+		bool bFirstSampler = true;
+		for (TStringSet::iterator IterSamplers = Samplers.begin(); IterSamplers != Samplers.end(); ++IterSamplers)
+		{
+		if (bFirstSampler)
+		{
+		bFirstSampler = false;
+		}
+		else
+		{
+		ralloc_asprintf_append(buffer, ",");
+		}
+
+		ralloc_asprintf_append(buffer, "%s", IterSamplers->c_str());
+		}
+		ralloc_asprintf_append(buffer, ")");
+		}
+		}
+		*/
+	}
+
+	bool PrintPackedUniforms(bool bPrintArrayType, char ArrayType, _mesa_glsl_parse_state::TUniformList& Uniforms, bool bFlattenUniformBuffers, bool NeedsComma)
+	{
+		bool bPrintHeader = true;
+		for (glsl_packed_uniform& Uniform : Uniforms)
+		{
+			if (!bFlattenUniformBuffers || Uniform.CB_PackedSampler.empty())
+			{
+				if (bPrintArrayType && bPrintHeader)
+				{
+					ralloc_asprintf_append(buffer, "%s%c[",
+						NeedsComma ? "," : "",
+						ArrayType);
+					bPrintHeader = false;
+					NeedsComma = false;
+				}
+				ralloc_asprintf_append(
+					buffer,
+					"%s%s(%u:%u)",
+					NeedsComma ? "," : "",
+					Uniform.Name.c_str(),
+					Uniform.offset,
+					Uniform.num_components
+					);
+				NeedsComma = true;
+			}
+		}
+
+		if (bPrintArrayType && !bPrintHeader)
+		{
+			ralloc_asprintf_append(buffer, "]");
+		}
+
+		return NeedsComma;
+	}
+
+	void PrintPackedGlobals(_mesa_glsl_parse_state* State)
+	{
+		//	@PackedGlobals: Global0(DestArrayType, DestOffset, SizeInFloats), Global1(DestArrayType, DestOffset, SizeInFloats), ...
+		bool bNeedsHeader = true;
+		bool bNeedsComma = false;
+		for (auto& Pair : State->GlobalPackedArraysMap)
+		{
+			char ArrayType = Pair.first;
+			if (ArrayType != EArrayType_Image && ArrayType != EArrayType_Sampler)
+			{
+				_mesa_glsl_parse_state::TUniformList& Uniforms = Pair.second;
+				check(!Uniforms.empty());
+
+				for (auto Iter = Uniforms.begin(); Iter != Uniforms.end(); ++Iter)
+				{
+					glsl_packed_uniform& Uniform = *Iter;
+					if (!State->bFlattenUniformBuffers || Uniform.CB_PackedSampler.empty())
+					{
+						if (bNeedsHeader)
+						{
+							ralloc_asprintf_append(buffer, "// @PackedGlobals: ");
+							bNeedsHeader = false;
+						}
+
+						ralloc_asprintf_append(
+							buffer,
+							"%s%s(%c:%u,%u)",
+							bNeedsComma ? "," : "",
+							Uniform.Name.c_str(),
+							ArrayType,
+							Uniform.offset,
+							Uniform.num_components
+							);
+						bNeedsComma = true;
+					}
+				}
+			}
+		}
+
+		if (!bNeedsHeader)
+		{
+			ralloc_asprintf_append(buffer, "\n");
+		}
+	}
+
+	void PrintPackedUniformBuffers(_mesa_glsl_parse_state* State, bool bGroupFlattenedUBs)
+	{
+		// @PackedUB: UniformBuffer0(SourceIndex0): Member0(SourceOffset,SizeInFloats),Member1(SourceOffset,SizeInFloats), ...
+		// @PackedUB: UniformBuffer1(SourceIndex1): Member0(SourceOffset,SizeInFloats),Member1(SourceOffset,SizeInFloats), ...
+		// ...
+
+		// First find all used CBs (since we lost that info during flattening)
+		TStringSet UsedCBs;
+		for (auto IterCB = State->CBPackedArraysMap.begin(); IterCB != State->CBPackedArraysMap.end(); ++IterCB)
+		{
+			for (auto Iter = IterCB->second.begin(); Iter != IterCB->second.end(); ++Iter)
+			{
+				_mesa_glsl_parse_state::TUniformList& Uniforms = Iter->second;
+				for (auto IterU = Uniforms.begin(); IterU != Uniforms.end(); ++IterU)
+				{
+					if (!IterU->CB_PackedSampler.empty())
+					{
+						check(IterCB->first == IterU->CB_PackedSampler);
+						UsedCBs.insert(IterU->CB_PackedSampler);
+					}
+				}
+			}
+		}
+
+		check(UsedCBs.size() == State->CBPackedArraysMap.size());
+
+		// Now get the CB index based off source declaration order, and print an info line for each, while creating the mem copy list
+		unsigned CBIndex = 0;
+		TCBDMARangeMap CBRanges;
+		for (unsigned i = 0; i < State->num_uniform_blocks; i++)
+		{
+			const glsl_uniform_block* block = State->uniform_blocks[i];
+			if (UsedCBs.find(block->name) != UsedCBs.end())
+			{
+				bool bNeedsHeader = true;
+
+				// Now the members for this CB
+				bool bNeedsComma = false;
+				auto IterPackedArrays = State->CBPackedArraysMap.find(block->name);
+				check(IterPackedArrays != State->CBPackedArraysMap.end());
+				for (auto Iter = IterPackedArrays->second.begin(); Iter != IterPackedArrays->second.end(); ++Iter)
+				{
+					char ArrayType = Iter->first;
+					check(ArrayType != EArrayType_Image && ArrayType != EArrayType_Sampler);
+
+					_mesa_glsl_parse_state::TUniformList& Uniforms = Iter->second;
+					for (auto IterU = Uniforms.begin(); IterU != Uniforms.end(); ++IterU)
+					{
+						glsl_packed_uniform& Uniform = *IterU;
+						if (Uniform.CB_PackedSampler == block->name)
+						{
+							if (bNeedsHeader)
+							{
+								ralloc_asprintf_append(buffer, "// @PackedUB: %s(%d): ",
+									block->name,
+									CBIndex);
+								bNeedsHeader = false;
+							}
+
+							ralloc_asprintf_append(buffer, "%s%s(%u,%u)",
+								bNeedsComma ? "," : "",
+								Uniform.Name.c_str(),
+								Uniform.OffsetIntoCBufferInFloats,
+								Uniform.SizeInFloats);
+
+							bNeedsComma = true;
+							unsigned SourceOffset = Uniform.OffsetIntoCBufferInFloats;
+							unsigned DestOffset = Uniform.offset;
+							unsigned Size = Uniform.SizeInFloats;
+							unsigned DestCBIndex = bGroupFlattenedUBs ? std::distance(UsedCBs.begin(), UsedCBs.find(block->name)) : 0;
+							unsigned DestCBPrecision = ArrayType;
+							InsertRange(CBRanges, CBIndex, SourceOffset, Size, DestCBIndex, DestCBPrecision, DestOffset);
+						}
+					}
+				}
+
+				if (!bNeedsHeader)
+				{
+					ralloc_asprintf_append(buffer, "\n");
+				}
+
+				CBIndex++;
+			}
+		}
+
+		//DumpSortedRanges(SortRanges(CBRanges));
+
+		// @PackedUBCopies: SourceArray:SourceOffset-DestArray:DestOffset,SizeInFloats;SourceArray:SourceOffset-DestArray:DestOffset,SizeInFloats,...
+		bool bFirst = true;
+		for (auto& Pair : CBRanges)
+		{
+			TDMARangeList& List = Pair.second;
+			for (auto IterList = List.begin(); IterList != List.end(); ++IterList)
+			{
+				if (bFirst)
+				{
+					ralloc_asprintf_append(buffer, bGroupFlattenedUBs ? "// @PackedUBCopies: " : "// @PackedUBGlobalCopies: ");
+					bFirst = false;
+				}
+				else
+				{
+					ralloc_asprintf_append(buffer, ",");
+				}
+
+				if (bGroupFlattenedUBs)
+				{
+					ralloc_asprintf_append(buffer, "%d:%d-%d:%c:%d:%d", IterList->SourceCB, IterList->SourceOffset, IterList->DestCBIndex, IterList->DestCBPrecision, IterList->DestOffset, IterList->Size);
+				}
+				else
+				{
+					check(IterList->DestCBIndex == 0);
+					ralloc_asprintf_append(buffer, "%d:%d-%c:%d:%d", IterList->SourceCB, IterList->SourceOffset, IterList->DestCBPrecision, IterList->DestOffset, IterList->Size);
+				}
+			}
+		}
+
+		if (!bFirst)
+		{
+			ralloc_asprintf_append(buffer, "\n");
+		}
+	}
+
+	void PrintPackedUniforms(_mesa_glsl_parse_state* State, bool bGroupFlattenedUBs)
+	{
+		PrintPackedGlobals(State);
+
+		if (State->bFlattenUniformBuffers && !State->CBuffersOriginal.empty())
+		{
+			PrintPackedUniformBuffers(State, bGroupFlattenedUBs);
+		}
+	}
+
+	/**
+	 * Print a list of external variables.
+	 */
+	void print_extern_vars(_mesa_glsl_parse_state* State, exec_list* extern_vars)
+	{
+		const char *type_str[GLSL_TYPE_MAX] = { "u", "i", "f", "f", "b", "t", "?", "?", "?", "?", "s", "os", "im", "ip", "op" };
+		const char *col_str[] = { "", "", "2x", "3x", "4x" };
+		const char *row_str[] = { "", "1", "2", "3", "4" };
+
+		check( sizeof(type_str)/sizeof(char*) == GLSL_TYPE_MAX);
+
+		bool need_comma = false;
+		foreach_iter(exec_list_iterator, iter, *extern_vars)
+		{
+			ir_variable* var = ((extern_var*)iter.get())->var;
+			const glsl_type* type = var->type;
+			if (!strcmp(var->name,"gl_in"))
+			{
+				// Ignore it, as we can't properly frame this information in current format, and it's not used anyway for geometry shaders
+				continue;
+			}
+			if (!strncmp(var->name,"in_",3) || !strncmp(var->name,"out_",4))
+			{
+				if (type->is_record())
+				{
+					// This is the specific case for GLSL >= 150, as we generate a struct with a member for each interpolator (which we still want to count)
+					if (type->length != 1)
+					{
+						_mesa_glsl_warning(State, "Found a complex structure as in/out, counting is not implemented yet...\n");
+						continue;
+					}
+
+					type = type->fields.structure->type;
+				}
+
+				// In and out variables may be packed in structures, or array of structures.
+				// But we're interested only in those that aren't, ie. inputs for vertex shader and outputs for pixel shader.
+				if (type->is_array() || type->is_record())
+				{
+					continue;
+				}
+			}
+			bool is_array = type->is_array();
+			int array_size = is_array ? type->length : 0;
+			if (is_array)
+			{
+				type = type->fields.array;
+			}
+			ralloc_asprintf_append(buffer, "%s%s%s%s",
+				need_comma ? "," : "",
+				type_str[type->base_type],
+				col_str[type->matrix_columns],
+				row_str[type->vector_elements]);
+			if (is_array)
+			{
+				ralloc_asprintf_append(buffer, "[%u]", array_size);
+			}
+			ralloc_asprintf_append(buffer, ":%s", var->name);
+			need_comma = true;
+		}
+	}
+
+	/**
+	 * Print the input/output signature for this shader.
+	 */
+	void print_signature(_mesa_glsl_parse_state *state, bool bGroupFlattenedUBs)
+	{
+		if (!input_variables.is_empty())
+		{
+			ralloc_asprintf_append(buffer, "// @Inputs: ");
+			print_extern_vars(state, &input_variables);
+			ralloc_asprintf_append(buffer, "\n");
+		}
+		if (!output_variables.is_empty())
+		{
+			ralloc_asprintf_append(buffer, "// @Outputs: ");
+			print_extern_vars(state, &output_variables);
+			ralloc_asprintf_append(buffer, "\n");
+		}
+		if (state->num_uniform_blocks > 0 && !state->bFlattenUniformBuffers)
+		{
+			bool first = true;
+			int Index = 0;
+			for (unsigned i = 0; i < state->num_uniform_blocks; i++)
+			{
+				const glsl_uniform_block* block = state->uniform_blocks[i];
+				if (hash_table_find(used_uniform_blocks, block->name))
+				{
+					ralloc_asprintf_append(buffer, "%s%s(%d)",
+						first ? "// @UniformBlocks: " : ",",
+						block->name, Index);
+					first = false;
+					++Index;
+				}
+			}
+			if (!first)
+			{
+				ralloc_asprintf_append(buffer, "\n");
+			}
+		}
+
+		if (state->has_packed_uniforms)
+		{
+			PrintPackedUniforms(state, bGroupFlattenedUBs);
+
+			if (!state->GlobalPackedArraysMap[EArrayType_Sampler].empty())
+			{
+				ralloc_asprintf_append(buffer, "// @Samplers: ");
+				PrintPackedSamplers(
+					state->GlobalPackedArraysMap[EArrayType_Sampler],
+					state->TextureToSamplerMap
+					);
+				ralloc_asprintf_append(buffer, "\n");
+			}
+
+			if (!state->GlobalPackedArraysMap[EArrayType_Image].empty())
+			{
+				ralloc_asprintf_append(buffer, "// @UAVs: ");
+				PrintPackedUniforms(
+					false,
+					EArrayType_Image,
+					state->GlobalPackedArraysMap[EArrayType_Image],
+					false,
+					false
+					);
+				ralloc_asprintf_append(buffer, "\n");
+			}
+		}
+		else
+		{
+			if (!uniform_variables.is_empty())
+			{
+				ralloc_asprintf_append(buffer, "// @Uniforms: ");
+				print_extern_vars(state, &uniform_variables);
+				ralloc_asprintf_append(buffer, "\n");
+			}
+			if (!sampler_variables.is_empty())
+			{
+				ralloc_asprintf_append(buffer, "// @Samplers: ");
+				print_extern_vars(state, &sampler_variables);
+				ralloc_asprintf_append(buffer, "\n");
+			}
+			if (!image_variables.is_empty())
+			{
+				ralloc_asprintf_append(buffer, "// @UAVs: ");
+				print_extern_vars(state, &image_variables);
+				ralloc_asprintf_append(buffer, "\n");
+			}
+		}
+	}
+
+	/**
+	 * Print the layout directives for this shader.
+	 */
+	void print_layout(_mesa_glsl_parse_state *state)
+	{
+		if (early_depth_stencil)
+		{
+			ralloc_asprintf_append(buffer, "layout(early_fragment_tests) in;\n");
+		}
+		if (state->target == compute_shader )
+		{
+			ralloc_asprintf_append(buffer, "layout( local_size_x = %d, "
+				"local_size_y = %d, local_size_z = %d ) in;\n", wg_size_x,
+				wg_size_y, wg_size_z );
+		}
+
+		if(state->target == tessellation_control_shader)
+		{
+			ralloc_asprintf_append(buffer, "layout(vertices = %d) out;\n", tessellation.outputcontrolpoints);
+		}
+
+		if(state->target == tessellation_evaluation_shader)
+		{
+
+			std::stringstream str;
+
+			switch (tessellation.outputtopology)
+			{
+				// culling is inverted, see TranslateCullMode in the OpenGL and D3D11 RHI
+			case GLSL_OUTPUTTOPOLOGY_POINT:
+				str << "point_mode";
+				break;
+			case GLSL_OUTPUTTOPOLOGY_LINE:
+				str << "iso_lines";
+				break;
+
+			default:
+			case GLSL_OUTPUTTOPOLOGY_NONE:
+			case GLSL_OUTPUTTOPOLOGY_TRIANGLE_CW:
+				str << "triangles, ccw";
+				break;
+			case GLSL_OUTPUTTOPOLOGY_TRIANGLE_CCW:
+				str << "triangles, cw";
+				break;
+			}
+
+			switch (tessellation.partitioning)
+			{
+			default:
+			case GLSL_PARTITIONING_NONE:
+			case GLSL_PARTITIONING_INTEGER:
+				str << ", equal_spacing";
+				break;
+			case GLSL_PARTITIONING_FRACTIONAL_EVEN:
+				str << ", fractional_even_spacing";
+				break;
+			case GLSL_PARTITIONING_FRACTIONAL_ODD:
+				str << ", fractional_odd_spacing";
+				break;
+				// that assumes that the hull/control shader clamps the tessellation factors to be power of two
+			case GLSL_PARTITIONING_POW2:
+				str << ", equal_spacing";
+				break;
+			}
+			ralloc_asprintf_append(buffer, "layout(%s) in;\n", str.str().c_str());
+		}
+#if 0
+		if(state->target == tessellation_evaluation_shader || state->target == tessellation_control_shader)
+		{
+			ralloc_asprintf_append(buffer, "/* DEBUG DUMP\n");
+
+			ralloc_asprintf_append(buffer, "tessellation.domain =  %s \n",  DomainStrings[tessellation.domain] );
+			ralloc_asprintf_append(buffer, "tessellation.outputtopology =  %s \n", OutputTopologyStrings[tessellation.outputtopology] );
+			ralloc_asprintf_append(buffer, "tessellation.partitioning =  %s \n", PartitioningStrings[tessellation.partitioning]);
+			ralloc_asprintf_append(buffer, "tessellation.maxtessfactor =  %f \n", tessellation.maxtessfactor );
+			ralloc_asprintf_append(buffer, "tessellation.outputcontrolpoints =  %d \n", tessellation.outputcontrolpoints );
+			ralloc_asprintf_append(buffer, "tessellation.patchconstantfunc =  %s \n", tessellation.patchconstantfunc);
+			ralloc_asprintf_append(buffer, " */\n");
+		}
+#endif		
+	}
+
+	void print_extensions(bool bUsesFramebufferFetchES2, bool bUsesES31Extensions)
+	{
+		if (bUsesES2TextureLODExtension)
+		{
+			ralloc_asprintf_append(buffer, "#ifndef DONTEMITEXTENSIONSHADERTEXTURELODENABLE\n");
+			ralloc_asprintf_append(buffer, "#extension GL_EXT_shader_texture_lod : enable\n");
+			ralloc_asprintf_append(buffer, "#endif\n");
+		}
+
+		if (bUsesDXDY && bIsES)
+		{
+			ralloc_asprintf_append(buffer, "#extension GL_OES_standard_derivatives : enable\n");
+		}
+
+		if (bUsesFramebufferFetchES2)
+		{
+			ralloc_asprintf_append(buffer, "\n#ifdef GL_EXT_shader_framebuffer_fetch\n");
+			ralloc_asprintf_append(buffer, "#extension GL_EXT_shader_framebuffer_fetch : enable\n");
+			ralloc_asprintf_append(buffer, "#endif\n");
+		}
+
+		if (bUsesES31Extensions)
+		{
+			ralloc_asprintf_append(buffer, "#extension GL_EXT_gpu_shader5 : enable\n");
+			ralloc_asprintf_append(buffer, "#extension GL_EXT_texture_buffer : enable\n");
+			ralloc_asprintf_append(buffer, "#extension GL_EXT_texture_cube_map_array : enable\n");
+			ralloc_asprintf_append(buffer, "#extension GL_EXT_shader_io_blocks : enable\n");
+
+			if (ShaderTarget == geometry_shader)
+			{
+				ralloc_asprintf_append(buffer, "#extension GL_EXT_geometry_shader : enable\n");
+			}
+
+			if (ShaderTarget == tessellation_control_shader || ShaderTarget == tessellation_evaluation_shader)
+			{
+				ralloc_asprintf_append(buffer, "#extension GL_EXT_tessellation_shader : enable\n");
+	}
+
+		}
+
+	}
+
+public:
+
+	/** Constructor. */
+	ir_gen_glsl_visitor(bool bInIsES, _mesa_glsl_parser_targets InShaderTarget, bool bInGenerateLayoutLocations)
+		: early_depth_stencil(false)
+		, bIsES(bInIsES)
+		, ShaderTarget(InShaderTarget)
+		, bGenerateLayoutLocations(bInGenerateLayoutLocations)
+		, buffer(0)
+		, indentation(0)
+		, scope_depth(0)
+		, temp_id(0)
+		, global_id(0)
+		, needs_semicolon(false)
+		, should_print_uint_literals_as_ints(false)
+		, loop_count(0)
+		, bUsesES2TextureLODExtension(false)
+		, bUsesDXDY(false)
+	{
+		printable_names = hash_table_ctor(32, hash_table_pointer_hash, hash_table_pointer_compare);
+		used_structures = hash_table_ctor(32, hash_table_pointer_hash, hash_table_pointer_compare);
+		used_uniform_blocks = hash_table_ctor(32, hash_table_string_hash, hash_table_string_compare);
+	}
+
+	/** Destructor. */
+	virtual ~ir_gen_glsl_visitor()
+	{
+		hash_table_dtor(printable_names);
+		hash_table_dtor(used_structures);
+		hash_table_dtor(used_uniform_blocks);
+	}
+
+	/**
+	 * Executes the visitor on the provided ir.
+	 * @returns the GLSL source code generated.
+	 */
+	const char* run(exec_list* ir, _mesa_glsl_parse_state* state, bool bGroupFlattenedUBs)
+	{
+		mem_ctx = ralloc_context(NULL);
+
+		char* code_buffer = ralloc_asprintf(mem_ctx, "");
+		buffer = &code_buffer;
+
+		if (bIsES && !(ShaderTarget == vertex_shader))
+		{
+			// TODO: Improve this...
+			
+			const char* DefaultPrecision = GDefaultPrecisionIsHalf ? "mediump" : "highp";
+			ralloc_asprintf_append(buffer, "precision %s float;\n", DefaultPrecision);
+			ralloc_asprintf_append(buffer, "precision %s int;\n", DefaultPrecision);
+			ralloc_asprintf_append(buffer, "\n#ifndef DONTEMITSAMPLERDEFAULTPRECISION\n"); 
+			ralloc_asprintf_append(buffer, "precision %s sampler2D;\n", DefaultPrecision);
+			ralloc_asprintf_append(buffer, "precision %s samplerCube;\n\n", DefaultPrecision);
+			ralloc_asprintf_append(buffer, "#endif\n");
+		}
+
+		if ((state->language_version == 310) && (ShaderTarget == fragment_shader))
+		{
+			ralloc_asprintf_append(buffer, "precision %s float;\n", "highp");
+			ralloc_asprintf_append(buffer, "precision %s int;\n", "highp");
+		}
+
+		// FramebufferFetchES2 'intrinsic'
+		bool bUsesFramebufferFetchES2 = UsesFramebufferFetchES2(ir);
+		if (bUsesFramebufferFetchES2)
+		{
+			ralloc_asprintf_append(buffer, "\n#ifdef GL_EXT_shader_framebuffer_fetch\n");
+			ralloc_asprintf_append(buffer, "vec4 FramebufferFetchES2() { return gl_LastFragData[0]; }\n");
+			ralloc_asprintf_append(buffer, "#else\n");
+			ralloc_asprintf_append(buffer, "vec4 FramebufferFetchES2() { return vec4(65000.0, 65000.0, 65000.0, 65000.0); }\n");
+			ralloc_asprintf_append(buffer, "#endif\n\n");
+		}
+
+		foreach_iter(exec_list_iterator, iter, *ir)
+		{
+			ir_instruction *inst = (ir_instruction *)iter.get();
+			do_visit(inst);
+		}
+		buffer = 0;
+
+		char* decl_buffer = ralloc_asprintf(mem_ctx, "");
+		buffer = &decl_buffer;
+		declare_structs(state);
+		buffer = 0;
+
+		char* signature = ralloc_asprintf(mem_ctx, "");
+		buffer = &signature;
+		print_signature(state, bGroupFlattenedUBs);
+		buffer = 0;
+
+		const char* geometry_layouts = "";
+		if (state->maxvertexcount>0)
+		{
+			check(state->geometryinput>0);
+			check(state->outputstream_type>0);
+			geometry_layouts = ralloc_asprintf(
+				mem_ctx,
+				"\nlayout(%s) in;\nlayout(%s, max_vertices = %d) out;\n\n",
+				GeometryInputStrings[state->geometryinput],
+				OutputStreamTypeStrings[state->outputstream_type],
+				state->maxvertexcount);
+		}
+		char* layout = ralloc_asprintf(mem_ctx, "");
+		buffer = &layout;
+		print_layout(state);
+		buffer = 0;
+
+		char* Extensions = ralloc_asprintf(mem_ctx, "");
+		buffer = &Extensions;
+		print_extensions(bUsesFramebufferFetchES2, state->language_version == 310);
+		buffer = 0;
+
+		char* full_buffer = ralloc_asprintf(
+			state,
+			"// Compiled by HLSLCC %d.%d\n%s#version %u %s\n%s%s%s%s%s\n",
+			HLSLCC_VersionMajor, HLSLCC_VersionMinor,
+			signature,
+			state->language_version,
+			state->language_version == 310 ? "es" : "",
+			Extensions,
+			geometry_layouts,
+			layout,
+			decl_buffer,
+			code_buffer
+			);
+		ralloc_free(mem_ctx);
+
+		return full_buffer;
+	}
+};
+
+struct FBreakPrecisionChangesVisitor : public ir_rvalue_visitor
+{
+	_mesa_glsl_parse_state* State;
+	FBreakPrecisionChangesVisitor(_mesa_glsl_parse_state* InState) : State(InState) {}
+
+	virtual void handle_rvalue(ir_rvalue** RValuePtr) override
+	{
+		if (!RValuePtr || !*RValuePtr)
+		{
+			return;
+		}
+		bool bGenerateNewVar = false;
+		auto* RValue = *RValuePtr;
+		auto* Expression = RValue->as_expression();
+		auto* Constant = RValue->as_constant();
+		if (Expression)
+		{
+			if (GDefaultPrecisionIsHalf)
+			{
+				switch (Expression->operation)
+				{
+					case ir_unop_i2f:
+					case ir_unop_b2f:
+					case ir_unop_u2f:
+						bGenerateNewVar = GDefaultPrecisionIsHalf;
+						break;
+
+					case ir_unop_i2h:
+					case ir_unop_b2h:
+					case ir_unop_u2h:
+						bGenerateNewVar = !GDefaultPrecisionIsHalf;
+						break;
+
+					case ir_unop_h2f:
+					case ir_unop_f2h:
+						if (!Expression->operands[0]->as_texture())
+						{
+							bGenerateNewVar = true;
+						}
+						break;
+				}
+			}
+		}
+		else if (Constant)
+		{
+/*
+			if ((GDefaultPrecisionIsHalf && Constant->type->base_type == GLSL_TYPE_HALF) ||
+				(!GDefaultPrecisionIsHalf && Constant->type->base_type == GLSL_TYPE_FLOAT))
+			{
+				bGenerateNewVar = true;
+			}
+*/
+		}
+		if (bGenerateNewVar)
+		{
+			auto* NewVar = new(State)ir_variable(RValue->type, nullptr, ir_var_temporary);
+			auto* NewAssignment = new(State)ir_assignment(new(State)ir_dereference_variable(NewVar), RValue);
+			*RValuePtr = new(State)ir_dereference_variable(NewVar);
+			base_ir->insert_before(NewVar);
+			base_ir->insert_before(NewAssignment);
+		}
+	}
+};
+
+void ir_gen_glsl_visitor::AddTypeToUsedStructs(const glsl_type* type)
+{
+	if (type->base_type == GLSL_TYPE_STRUCT)
+	{
+		if (hash_table_find(used_structures, type) == NULL)
+		{
+			hash_table_insert(used_structures, (void*)type, type);
+		}
+	}
+
+	if (type->base_type == GLSL_TYPE_ARRAY && type->fields.array->base_type == GLSL_TYPE_STRUCT)
+	{
+		if (hash_table_find(used_structures, type->fields.array) == NULL)
+		{
+			hash_table_insert(used_structures, (void*)type->fields.array, type->fields.array);
+		}
+	}
+
+	if ((type->base_type == GLSL_TYPE_INPUTPATCH || type->base_type == GLSL_TYPE_OUTPUTPATCH) && type->inner_type->base_type == GLSL_TYPE_STRUCT)
+	{
+		if (hash_table_find(used_structures, type->inner_type) == NULL)
+		{
+			hash_table_insert(used_structures, (void*)type->inner_type, type->inner_type);
+		}
+	}
+}
+
+char* FGlslCodeBackend::GenerateCode(exec_list* ir, _mesa_glsl_parse_state* state, EHlslShaderFrequency Frequency)
+{
+	FixRedundantCasts(ir);
+//IRDump(ir);
+
+	FBreakPrecisionChangesVisitor BreakPrecisionChangesVisitor(state);
+	BreakPrecisionChangesVisitor.run(ir);
+
+	const bool bGroupFlattenedUBs = ((HlslCompileFlags & HLSLCC_GroupFlattenedUniformBuffers) == HLSLCC_GroupFlattenedUniformBuffers);
+	const bool bGenerateLayoutLocations = state->bGenerateLayoutLocations;
+	ir_gen_glsl_visitor visitor(state->bGenerateES, state->target, bGenerateLayoutLocations);
+	const char* code = visitor.run(ir, state, bGroupFlattenedUBs);
+	return _strdup(code);
+}
+
+
+// Verify if SampleLevel() is used
+struct SPromoteSampleLevelES2 : public ir_hierarchical_visitor
+{
+	_mesa_glsl_parse_state* ParseState;
+	const bool bIsVertexShader;
+	SPromoteSampleLevelES2(_mesa_glsl_parse_state* InParseState, bool bInIsVertexShader) :
+		ParseState(InParseState),
+		bIsVertexShader(bInIsVertexShader)
+	{
+	}
+
+	virtual ir_visitor_status visit_leave(ir_texture* IR) override
+	{
+		if (IR->op == ir_txl)
+		{
+			if (bIsVertexShader)
+			{
+				YYLTYPE loc;
+				loc.first_column = IR->SourceLocation.Column;
+				loc.first_line = IR->SourceLocation.Line;
+				loc.source_file = IR->SourceLocation.SourceFile.c_str();
+				_mesa_glsl_error(&loc, ParseState, "Vertex texture fetch currently not supported on GLSL ES\n");
+			}
+			else
+			{
+				//@todo-mobile: allowing lod texture functions for now, as they are supported on some devices via glsl extension.
+				// http://www.khronos.org/registry/gles/extensions/EXT/EXT_shader_texture_lod.txt
+				// Compat work will be required for devices which do not support it.
+				/*
+				_mesa_glsl_warning(ParseState, "%s(%d, %d) Converting SampleLevel() to Sample()\n", IR->SourceLocation.SourceFile.c_str(), IR->SourceLocation.Line, IR->SourceLocation.Column);
+				IR->op = ir_tex;
+				*/
+			}
+		}
+
+		if (IR->offset)
+		{
+			YYLTYPE loc;
+			loc.first_column = IR->SourceLocation.Column;
+			loc.first_line = IR->SourceLocation.Line;
+			loc.source_file = IR->SourceLocation.SourceFile.c_str();
+			_mesa_glsl_error(&loc, ParseState, "Texture offset not supported on GLSL ES\n");
+		}
+
+		return visit_continue;
+	}
+}; 
+
+
+// Converts an array index expression using an integer input attribute, to a float input attribute using a conversion to int
+struct SConvertIntVertexAttributeES2 : public ir_hierarchical_visitor
+{
+	_mesa_glsl_parse_state* ParseState;
+	exec_list* FunctionBody;
+	int InsideArrayDeref;
+	std::map<ir_variable*, ir_variable*> ConvertedVarMap;
+
+	SConvertIntVertexAttributeES2(_mesa_glsl_parse_state* InParseState, exec_list* InFunctionBody) : ParseState(InParseState), FunctionBody(InFunctionBody), InsideArrayDeref(0)
+	{
+	}
+
+	virtual ir_visitor_status visit_enter(ir_dereference_array* DeRefArray) override
+	{
+		// Break the array dereference so we know we want to modify the array index part
+		auto Result = ir_hierarchical_visitor::visit_enter(DeRefArray);
+		++InsideArrayDeref;
+		DeRefArray->array_index->accept(this);
+		--InsideArrayDeref;
+
+		return visit_continue;
+	}
+
+	virtual ir_visitor_status visit(ir_dereference_variable* DeRefVar) override
+	{
+		if (InsideArrayDeref > 0)
+		{
+			ir_variable* SourceVar = DeRefVar->var;
+			if (SourceVar->mode == ir_var_in)
+			{
+				// First time it still is an integer, so add the temporary and a conversion, and switch to float
+				if (SourceVar->type->is_integer())
+				{
+					check(SourceVar->type->is_integer() && !SourceVar->type->is_matrix() && !SourceVar->type->is_array());
+
+					// Double check we haven't processed this
+					auto IterFound = ConvertedVarMap.find(SourceVar);
+					check(IterFound == ConvertedVarMap.end());
+
+					// New temp var
+					ir_variable* NewVar = new(ParseState)ir_variable(SourceVar->type, NULL, ir_var_temporary);
+					base_ir->insert_before(NewVar);
+
+					// Switch original type to float
+					SourceVar->type = glsl_type::get_instance(GLSL_TYPE_FLOAT, SourceVar->type->vector_elements, 1);
+
+					// Convert float to int
+					ir_dereference_variable* NewSourceDeref = new(ParseState)ir_dereference_variable(SourceVar);
+					ir_expression* NewCastExpression = new(ParseState)ir_expression(ir_unop_f2i, NewSourceDeref);
+					ir_assignment* NewAssigment = new(ParseState)ir_assignment(new(ParseState)ir_dereference_variable(NewVar), NewCastExpression);
+					base_ir->insert_before(NewAssigment);
+
+					// Add the entry and modify the original Var
+					ConvertedVarMap[SourceVar] = NewVar;
+					DeRefVar->var = NewVar;
+				}
+				else
+				{
+					auto IterFound = ConvertedVarMap.find(SourceVar);
+					if (IterFound != ConvertedVarMap.end())
+					{
+						DeRefVar->var = IterFound->second;
+					}
+				}
+			}
+		}
+
+		return ir_hierarchical_visitor::visit(DeRefVar);
+	}
+};
+
+
+bool FGlslCodeBackend::ApplyAndVerifyPlatformRestrictions(exec_list* Instructions, _mesa_glsl_parse_state* ParseState, EHlslShaderFrequency Frequency)
+{
+	if (ParseState->bGenerateES)
+	{
+		bool bIsVertexShader = (Frequency == HSF_VertexShader);
+
+		// Handle SampleLevel
+		{
+			SPromoteSampleLevelES2 Visitor(ParseState, bIsVertexShader);
+			Visitor.run(Instructions);
+		}
+
+		// Handle matrices (flatten to vectors so we can support non-sqaure)
+		ExpandMatricesIntoArrays(Instructions, ParseState);
+
+		// Handle integer vertex attributes used as array indices
+		if (bIsVertexShader)
+		{
+			SConvertIntVertexAttributeES2 ConvertIntVertexAttributeVisitor(ParseState, Instructions);
+			ConvertIntVertexAttributeVisitor.run(Instructions);
+		}
+	}
+
+	return true;
+}
+
+/** Qualifers that apply to semantics. */
+union FSemanticQualifier
+{
+	struct
+	{
+		unsigned bCentroid : 1;
+		unsigned InterpolationMode : 2;
+		unsigned bIsPatchConstant : 1;
+	} Fields;
+	unsigned Packed;
+
+	FSemanticQualifier() : Packed(0) {}
+};
+
+/** Information on system values. */
+struct FSystemValue
+{
+	const char* Semantic;
+	const glsl_type* Type;
+	const char* GlslName;
+	ir_variable_mode Mode;
+	bool bOriginUpperLeft;
+	bool bArrayVariable;
+	bool bApplyClipSpaceAdjustment;
+	bool bESOnly;
+};
+
+/** Vertex shader system values. */
+static FSystemValue VertexSystemValueTable[] =
+{
+	{ "SV_VertexID", glsl_type::int_type, "gl_VertexID", ir_var_in, false, false, false, false },
+	{ "SV_InstanceID", glsl_type::int_type, "gl_InstanceID", ir_var_in, false, false, false, false },
+	{ "SV_Position", glsl_type::vec4_type, "gl_Position", ir_var_out, false, false, true, false },
+	{ NULL, NULL, NULL, ir_var_auto, false, false, false, false }
+};
+
+/** Pixel shader system values. */
+static FSystemValue PixelSystemValueTable[] =
+{
+	{ "SV_Depth", glsl_type::float_type, "gl_FragDepth", ir_var_out, false, false, false, false },
+	{ "SV_Position", glsl_type::vec4_type, "gl_FragCoord", ir_var_in, true, false, false, false },
+	{ "SV_IsFrontFace", glsl_type::bool_type, "gl_FrontFacing", ir_var_in, false, false, true, false },
+	{ "SV_PrimitiveID", glsl_type::int_type, "gl_PrimitiveID", ir_var_in, false, false, false, false },
+	{ "SV_RenderTargetArrayIndex", glsl_type::int_type, "gl_Layer", ir_var_in, false, false, false, false },
+	{ "SV_Target0", glsl_type::half4_type, "gl_FragColor", ir_var_out, false, false, false, true },
+	{ NULL, NULL, NULL, ir_var_auto, false, false, false }
+};
+
+/** Geometry shader system values. */
+static FSystemValue GeometrySystemValueTable[] =
+{
+	{ "SV_VertexID", glsl_type::int_type, "gl_VertexID", ir_var_in, false, false, false, false },
+	{ "SV_InstanceID", glsl_type::int_type, "gl_InstanceID", ir_var_in, false, false, false, false },
+	{ "SV_Position", glsl_type::vec4_type, "gl_Position", ir_var_in, false, true, true, false },
+	{ "SV_Position", glsl_type::vec4_type, "gl_Position", ir_var_out, false, false, true, false },
+	{ "SV_RenderTargetArrayIndex", glsl_type::int_type, "gl_Layer", ir_var_out, false, false, false, false },
+	{ "SV_PrimitiveID", glsl_type::int_type, "gl_PrimitiveID", ir_var_out, false, false, false, false },
+	{ "SV_PrimitiveID", glsl_type::int_type, "gl_PrimitiveIDIn", ir_var_in, false, false, false, false },
+	{ NULL, NULL, NULL, ir_var_auto, false, false, false, false }
+};
+
+
+/** Hull shader system values. */
+static FSystemValue HullSystemValueTable[] =
+{
+	{ "SV_OutputControlPointID", glsl_type::int_type, "gl_InvocationID", ir_var_in, false, false, false, false },
+	{ NULL, NULL, NULL, ir_var_auto, false, false, false, false }
+};
+
+/** Domain shader system values. */
+static FSystemValue DomainSystemValueTable[] =
+{
+
+	// TODO : SV_DomainLocation has types float2 or float3 depending on the input topology
+	{ "SV_Position", glsl_type::vec4_type, "gl_Position", ir_var_in, false, true, true, false },
+	{ "SV_Position", glsl_type::vec4_type, "gl_Position", ir_var_out, false, false, true, false },
+	{ "SV_DomainLocation", glsl_type::vec3_type, "gl_TessCoord", ir_var_in, false, false, false, false },
+	{ NULL, NULL, NULL, ir_var_auto, false, false, false, false }
+};
+
+/** Compute shader system values. */
+static FSystemValue ComputeSystemValueTable[] =
+{
+	{ "SV_DispatchThreadID", glsl_type::uvec3_type, "gl_GlobalInvocationID", ir_var_in, false, false, false, false },
+	{ "SV_GroupID", glsl_type::uvec3_type, "gl_WorkGroupID", ir_var_in, false, false, false, false },
+	{ "SV_GroupIndex", glsl_type::uint_type, "gl_LocalInvocationIndex", ir_var_in, false, false, false, false },
+	{ "SV_GroupThreadID", glsl_type::uvec3_type, "gl_LocalInvocationID", ir_var_in, false, false, false, false },
+	{ NULL, NULL, NULL, ir_var_auto, false, false, false, false }
+};
+
+FSystemValue* SystemValueTable[HSF_FrequencyCount] =
+{
+	VertexSystemValueTable,
+	PixelSystemValueTable,
+	GeometrySystemValueTable,
+	HullSystemValueTable,
+	DomainSystemValueTable,
+	ComputeSystemValueTable
+};
+
+
+#define CUSTOM_LAYER_INDEX_SEMANTIC "HLSLCC_LAYER_INDEX"
+
+/**
+* Generate an input semantic.
+* @param Frequency - The shader frequency.
+* @param ParseState - Parse state.
+* @param Semantic - The semantic name to generate.
+* @param Type - Value type.
+* @param DeclInstructions - IR to which declarations may be added.
+* @returns reference to IR variable for the semantic.
+*/
+static ir_rvalue* GenShaderInputSemantic(
+	EHlslShaderFrequency Frequency,
+	_mesa_glsl_parse_state* ParseState,
+	const char* Semantic,
+	FSemanticQualifier InputQualifier,
+	const glsl_type* Type,
+	exec_list* DeclInstructions,
+	int SemanticArraySize,
+	int SemanticArrayIndex,
+	bool& ApplyClipSpaceAdjustment
+	)
+{
+	ir_variable* Variable = NULL;
+	if (Semantic && _strnicmp(Semantic, "SV_", 3) == 0)
+	{
+		FSystemValue* SystemValues = SystemValueTable[Frequency];
+		for (int i = 0; SystemValues[i].Semantic != NULL; ++i)
+		{
+			if (SystemValues[i].Mode == ir_var_in
+				&& (!SystemValues[i].bESOnly || ParseState->bGenerateES)
+				&& _stricmp(SystemValues[i].Semantic, Semantic) == 0)
+			{
+				if (SystemValues[i].bArrayVariable)
+				{
+					// Built-in array variable. Like gl_in[x].gl_Position.
+					// The variable for it has already been created in GenShaderInput().
+					/*ir_variable**/ Variable = ParseState->symbols->get_variable("gl_in");
+					check(Variable);
+					ir_dereference_variable* ArrayDeref = new(ParseState)ir_dereference_variable(Variable);
+					ir_dereference_array* StructDeref = new(ParseState)ir_dereference_array(
+						ArrayDeref,
+						new(ParseState)ir_constant((unsigned)SemanticArrayIndex)
+						);
+					ir_dereference_record* VariableDeref = new(ParseState)ir_dereference_record(
+						StructDeref,
+						SystemValues[i].GlslName
+						);
+					ApplyClipSpaceAdjustment = SystemValues[i].bApplyClipSpaceAdjustment;
+					// TO DO - in case of SV_ClipDistance, we need to defer appropriate index in variable too.
+					return VariableDeref;
+				}
+				else
+				{
+					// Built-in variable that shows up only once, like gl_FragCoord in fragment
+					// shader, or gl_PrimitiveIDIn in geometry shader. Unlike gl_in[x].gl_Position.
+					// Even in geometry shader input pass it shows up only once.
+
+					// Create it on first pass, ignore the call on others.
+					if (SemanticArrayIndex == 0)
+					{
+						ir_variable* Variable = new(ParseState)ir_variable(
+							SystemValues[i].Type,
+							SystemValues[i].GlslName,
+							ir_var_in
+							);
+						Variable->read_only = true;
+						Variable->origin_upper_left = SystemValues[i].bOriginUpperLeft;
+						DeclInstructions->push_tail(Variable);
+						ParseState->symbols->add_variable(Variable);
+						ir_dereference_variable* VariableDeref = new(ParseState)ir_dereference_variable(Variable);
+
+						if (ParseState->adjust_clip_space_dx11_to_opengl && SystemValues[i].bApplyClipSpaceAdjustment)
+						{
+							// incoming gl_FrontFacing. Make it (!gl_FrontFacing), due to vertical flip in OpenGL
+							return new(ParseState)ir_expression(ir_unop_logic_not, glsl_type::bool_type, VariableDeref, NULL);
+						}
+						else
+						{
+							return VariableDeref;
+						}
+					}
+					else
+					{
+						return NULL;
+					}
+				}
+			}
+		}
+	}
+
+	if (Variable == NULL && Frequency == HSF_DomainShader)
+	{
+		const int PrefixLength = 13;
+		if (_strnicmp(Semantic, "SV_TessFactor", PrefixLength) == 0
+			&& Semantic[PrefixLength] >= '0'
+			&& Semantic[PrefixLength] <= '3')
+		{
+			int OutputIndex = Semantic[PrefixLength] - '0';
+			Variable = new(ParseState)ir_variable(
+				Type,
+				ralloc_asprintf(ParseState, "gl_TessLevelOuter[%d]", OutputIndex),
+				ir_var_out
+				);
+		}
+	}
+
+	if (Variable == NULL && Frequency == HSF_DomainShader)
+	{
+		const int PrefixLength = 19;
+		if (_strnicmp(Semantic, "SV_InsideTessFactor", PrefixLength) == 0
+			&& Semantic[PrefixLength] >= '0'
+			&& Semantic[PrefixLength] <= '1')
+		{
+			int OutputIndex = Semantic[PrefixLength] - '0';
+			Variable = new(ParseState)ir_variable(
+				Type,
+				ralloc_asprintf(ParseState, "gl_TessLevelInner[%d]", OutputIndex),
+				ir_var_out
+				);
+		}
+		else if (_stricmp(Semantic, "SV_InsideTessFactor") == 0)
+		{
+			Variable = new(ParseState)ir_variable(
+				Type,
+				ralloc_asprintf(ParseState, "gl_TessLevelInner[0]"),
+				ir_var_out
+				);
+		}
+	}
+
+	if (Variable)
+	{
+		// Up to this point, variables aren't contained in structs
+		DeclInstructions->push_tail(Variable);
+		ParseState->symbols->add_variable(Variable);
+		Variable->centroid = InputQualifier.Fields.bCentroid;
+		Variable->interpolation = InputQualifier.Fields.InterpolationMode;
+		Variable->is_patch_constant = InputQualifier.Fields.bIsPatchConstant;
+		ir_rvalue* VariableDeref = new(ParseState)ir_dereference_variable(Variable);
+
+		return VariableDeref;
+	}
+
+	// If we're here, no built-in variables matched.
+
+	if (Semantic && _strnicmp(Semantic, "SV_", 3) == 0)
+	{
+		_mesa_glsl_warning(ParseState, "unrecognized system "
+			"value input '%s'", Semantic);
+	}
+
+	if (Frequency == HSF_VertexShader || ParseState->bGenerateES)
+	{
+		const char* Prefix = "in";
+		if (ParseState->bGenerateES && Frequency == HSF_PixelShader)
+		{
+			Prefix = "var";
+		}
+
+		// Vertex shader inputs don't get packed into structs that we'll later morph into interface blocks
+		if (ParseState->bGenerateES && Type->is_integer())
+		{
+			// Convert integer attributes to floats
+			ir_variable* Variable = new(ParseState)ir_variable(
+				Type,
+				ralloc_asprintf(ParseState, "%s_%s_I", Prefix, Semantic),
+				ir_var_temporary
+				);
+			Variable->centroid = InputQualifier.Fields.bCentroid;
+			Variable->interpolation = InputQualifier.Fields.InterpolationMode;
+			check(Type->is_vector() || Type->is_scalar());
+			check(Type->base_type == GLSL_TYPE_INT || Type->base_type == GLSL_TYPE_UINT);
+
+			// New float attribute
+			ir_variable* ReplacedAttributeVar = new (ParseState)ir_variable(glsl_type::get_instance(GLSL_TYPE_FLOAT, Variable->type->vector_elements, 1), ralloc_asprintf(ParseState, "%s_%s", Prefix, Semantic), ir_var_in);
+			ReplacedAttributeVar->read_only = true;
+			ReplacedAttributeVar->centroid = InputQualifier.Fields.bCentroid;
+			ReplacedAttributeVar->interpolation = InputQualifier.Fields.InterpolationMode;
+
+			// Convert to integer
+			ir_assignment* ConversionAssignment = new(ParseState)ir_assignment(
+				new(ParseState)ir_dereference_variable(Variable),
+				new(ParseState)ir_expression(
+				Type->base_type == GLSL_TYPE_INT ? ir_unop_f2i : ir_unop_f2u,
+				new (ParseState)ir_dereference_variable(ReplacedAttributeVar)
+				)
+				);
+
+			DeclInstructions->push_tail(ReplacedAttributeVar);
+			DeclInstructions->push_tail(Variable);
+			DeclInstructions->push_tail(ConversionAssignment);
+			ParseState->symbols->add_variable(Variable);
+			ParseState->symbols->add_variable(ReplacedAttributeVar);
+
+			ir_dereference_variable* VariableDeref = new(ParseState)ir_dereference_variable(ReplacedAttributeVar);
+			return VariableDeref;
+		}
+
+		// Regular attribute
+		Variable = new(ParseState)ir_variable(
+			Type,
+			ralloc_asprintf(ParseState, "%s_%s", Prefix, Semantic),
+			ir_var_in
+			);
+		Variable->read_only = true;
+		Variable->centroid = InputQualifier.Fields.bCentroid;
+		Variable->interpolation = InputQualifier.Fields.InterpolationMode;
+		Variable->is_patch_constant = InputQualifier.Fields.bIsPatchConstant;
+
+		// this is kinda hacky, but all the UE4 vertex shaders have ATTRIBUTE0 ... ATTRIBUTEN as the input semantics
+		// so we use the semantic index as the location.
+		if (Frequency == HSF_VertexShader && ParseState->bGenerateLayoutLocations)
+		{
+			const int PrefixLength = 9;
+			if ( (_strnicmp(Semantic, "ATTRIBUTE", PrefixLength) == 0) &&
+				 (Semantic[PrefixLength] >= '0') &&	 (Semantic[PrefixLength] <= '9') 
+				)
+			{
+				int AttributeIndex = atoi(Semantic + PrefixLength);
+
+				Variable->explicit_location = true;
+				Variable->location = AttributeIndex;
+			}
+			else
+			{
+#if DEBUG				
+				_mesa_glsl_warning(
+#else
+				_mesa_glsl_error(
+#endif
+					ParseState, "Vertex shader input semantic must be ATTRIBUTE and not \'%s\' in order to determine location/semantic index", Semantic);
+				
+			}
+		}
+
+		DeclInstructions->push_tail(Variable);
+		ParseState->symbols->add_variable(Variable);
+
+		ir_dereference_variable* VariableDeref = new(ParseState)ir_dereference_variable(Variable);
+		return VariableDeref;
+	}
+	else if (SemanticArrayIndex == 0)
+	{
+		// On first pass, create variable
+
+		glsl_struct_field *StructField = ralloc_array(ParseState, glsl_struct_field, 1);
+
+		memset(StructField, 0, sizeof(glsl_struct_field));
+		StructField[0].type = Type;
+		StructField[0].name = ralloc_strdup(ParseState, "Data");
+
+		const glsl_type* VariableType = glsl_type::get_record_instance(StructField, 1, ralloc_strdup(ParseState, Semantic));
+		if (SemanticArraySize)
+		{
+			// Pack it into an array too
+			VariableType = glsl_type::get_array_instance(VariableType, SemanticArraySize);
+		}
+
+		ir_variable* Variable = new(ParseState)ir_variable(VariableType, ralloc_asprintf(ParseState, "in_%s", Semantic), ir_var_in);
+		Variable->read_only = true;
+		Variable->is_interface_block = true;
+		Variable->centroid = InputQualifier.Fields.bCentroid;
+		Variable->interpolation = InputQualifier.Fields.InterpolationMode;
+		Variable->is_patch_constant = InputQualifier.Fields.bIsPatchConstant;
+		DeclInstructions->push_tail(Variable);
+		ParseState->symbols->add_variable(Variable);
+
+		if (ParseState->bGenerateLayoutLocations && !Variable->is_patch_constant)
+		{
+			Variable->location = ParseState->next_in_location_slot++;
+			Variable->explicit_location = true;
+		}
+
+		ir_rvalue* VariableDeref = new(ParseState)ir_dereference_variable(Variable);
+		if (SemanticArraySize)
+		{
+			// Deref inside array first
+			VariableDeref = new(ParseState)ir_dereference_array(VariableDeref, new(ParseState)ir_constant((unsigned)SemanticArrayIndex)
+				);
+		}
+		VariableDeref = new(ParseState)ir_dereference_record(VariableDeref, ralloc_strdup(ParseState, "Data"));
+		return VariableDeref;
+	}
+	else
+	{
+		// Array variable, not first pass. It already exists, get it.
+		ir_variable* Variable = ParseState->symbols->get_variable(ralloc_asprintf(ParseState, "in_%s", Semantic));
+		check(Variable);
+
+		ir_rvalue* VariableDeref = new(ParseState)ir_dereference_variable(Variable);
+		VariableDeref = new(ParseState)ir_dereference_array(VariableDeref, new(ParseState)ir_constant((unsigned)SemanticArrayIndex));
+		VariableDeref = new(ParseState)ir_dereference_record(VariableDeref, ralloc_strdup(ParseState, "Data"));
+		return VariableDeref;
+	}
+}
+
+/**
+* Generate an output semantic.
+* @param Frequency - The shader frequency.
+* @param ParseState - Parse state.
+* @param Semantic - The semantic name to generate.
+* @param Type - Value type.
+* @param DeclInstructions - IR to which declarations may be added.
+* @returns the IR variable for the semantic.
+*/
+static ir_rvalue* GenShaderOutputSemantic(
+	EHlslShaderFrequency Frequency,
+	_mesa_glsl_parse_state* ParseState,
+	const char* Semantic,
+	FSemanticQualifier OutputQualifier,
+	const glsl_type* Type,
+	exec_list* DeclInstructions,
+	const glsl_type** DestVariableType,
+	bool& ApplyClipSpaceAdjustment,
+	bool& ApplyClampPowerOfTwo
+	)
+{
+	FSystemValue* SystemValues = SystemValueTable[Frequency];
+	ir_variable* Variable = NULL;
+
+	if (Semantic && _strnicmp(Semantic, "SV_", 3) == 0)
+	{
+		for (int i = 0; SystemValues[i].Semantic != NULL; ++i)
+		{
+			if (!SystemValues[i].bESOnly || ParseState->bGenerateES)
+			{
+				if (SystemValues[i].Mode == ir_var_out
+					&& _stricmp(SystemValues[i].Semantic, Semantic) == 0)
+				{
+					Variable = new(ParseState)ir_variable(
+						SystemValues[i].Type,
+						SystemValues[i].GlslName,
+						ir_var_out
+						);
+					Variable->origin_upper_left = SystemValues[i].bOriginUpperLeft;
+					ApplyClipSpaceAdjustment = SystemValues[i].bApplyClipSpaceAdjustment;
+				}
+			}
+		}
+	}
+
+	if (Variable == NULL && Frequency == HSF_VertexShader)
+	{
+		const int PrefixLength = 15;
+		if (_strnicmp(Semantic, "SV_ClipDistance", PrefixLength) == 0
+			&& Semantic[PrefixLength] >= '0'
+			&& Semantic[PrefixLength] <= '9')
+		{
+			int OutputIndex = Semantic[15] - '0';
+			Variable = new(ParseState)ir_variable(
+				glsl_type::float_type,
+				ralloc_asprintf(ParseState, "gl_ClipDistance[%d]", OutputIndex),
+				ir_var_out
+				);
+		}
+	}
+
+	if (Variable == NULL && Frequency == HSF_PixelShader)
+	{
+		const int PrefixLength = 9;
+		if (_strnicmp(Semantic, "SV_Target", PrefixLength) == 0
+			&& Semantic[PrefixLength] >= '0'
+			&& Semantic[PrefixLength] <= '7')
+		{
+			int OutputIndex = Semantic[PrefixLength] - '0';
+			Variable = new(ParseState)ir_variable(
+				Type,
+				ralloc_asprintf(ParseState, "out_Target%d", OutputIndex),
+				ir_var_out
+				);
+
+			if (ParseState->bGenerateLayoutLocations)
+			{
+				Variable->explicit_location = true;
+				Variable->location = OutputIndex;
+		}
+	}
+	}
+
+	if (Variable == NULL && Frequency == HSF_HullShader)
+	{
+		const int PrefixLength = 13;
+		if (_strnicmp(Semantic, "SV_TessFactor", PrefixLength) == 0
+			&& Semantic[PrefixLength] >= '0'
+			&& Semantic[PrefixLength] <= '3')
+		{
+			int OutputIndex = Semantic[PrefixLength] - '0';
+			Variable = new(ParseState)ir_variable(
+				Type,
+				ralloc_asprintf(ParseState, "gl_TessLevelOuter[%d]", OutputIndex),
+				ir_var_out
+				);
+
+			ApplyClampPowerOfTwo = ParseState->tessellation.partitioning == GLSL_PARTITIONING_POW2;
+		}
+	}
+
+	if (Variable == NULL && Frequency == HSF_HullShader)
+	{
+		const int PrefixLength = 19;
+		if (_strnicmp(Semantic, "SV_InsideTessFactor", PrefixLength) == 0
+			&& Semantic[PrefixLength] >= '0'
+			&& Semantic[PrefixLength] <= '1')
+		{
+			int OutputIndex = Semantic[PrefixLength] - '0';
+			Variable = new(ParseState)ir_variable(
+				Type,
+				ralloc_asprintf(ParseState, "gl_TessLevelInner[%d]", OutputIndex),
+				ir_var_out
+				);
+
+			ApplyClampPowerOfTwo = ParseState->tessellation.partitioning == GLSL_PARTITIONING_POW2;
+		}
+		else if (_stricmp(Semantic, "SV_InsideTessFactor") == 0)
+		{
+			Variable = new(ParseState)ir_variable(
+				Type,
+				ralloc_asprintf(ParseState, "gl_TessLevelInner[0]"),
+				ir_var_out
+				);
+
+			ApplyClampPowerOfTwo = ParseState->tessellation.partitioning == GLSL_PARTITIONING_POW2;
+		}
+	}
+
+	if (Variable == NULL && ParseState->bGenerateES)
+	{
+		// Create a variable so that a struct will not get added
+		Variable = new(ParseState)ir_variable(Type, ralloc_asprintf(ParseState, "var_%s", Semantic), ir_var_out);
+	}
+
+	if (Variable)
+	{
+		// Up to this point, variables aren't contained in structs
+		*DestVariableType = Variable->type;
+		DeclInstructions->push_tail(Variable);
+		ParseState->symbols->add_variable(Variable);
+		Variable->centroid = OutputQualifier.Fields.bCentroid;
+		Variable->interpolation = OutputQualifier.Fields.InterpolationMode;
+		Variable->is_patch_constant = OutputQualifier.Fields.bIsPatchConstant;
+		ir_rvalue* VariableDeref = new(ParseState)ir_dereference_variable(Variable);
+		return VariableDeref;
+	}
+
+	if (Semantic && _strnicmp(Semantic, "SV_", 3) == 0)
+	{
+		_mesa_glsl_warning(ParseState, "unrecognized system value output '%s'",
+			Semantic);
+	}
+
+	*DestVariableType = Type;
+
+	// Create variable
+	glsl_struct_field *StructField = ralloc_array(ParseState, glsl_struct_field, 1);
+
+	memset(StructField, 0, sizeof(glsl_struct_field));
+	StructField[0].type = Type;
+	StructField[0].name = ralloc_strdup(ParseState, "Data");
+
+	const glsl_type* VariableType = glsl_type::get_record_instance(StructField, 1, ralloc_strdup(ParseState, Semantic));
+
+	if (Frequency == HSF_HullShader && !OutputQualifier.Fields.bIsPatchConstant)
+	{
+		VariableType = glsl_type::get_array_instance(VariableType, ParseState->tessellation.outputcontrolpoints);
+	}
+	Variable = new(ParseState)ir_variable(VariableType, ralloc_asprintf(ParseState, "out_%s", Semantic), ir_var_out);
+
+	Variable->centroid = OutputQualifier.Fields.bCentroid;
+	Variable->interpolation = OutputQualifier.Fields.InterpolationMode;
+	Variable->is_interface_block = true;
+	Variable->is_patch_constant = OutputQualifier.Fields.bIsPatchConstant;
+
+	if (ParseState->bGenerateLayoutLocations && !Variable->is_patch_constant)
+	{
+		Variable->location = ParseState->next_out_location_slot++;
+		Variable->explicit_location = true;
+	}
+
+	DeclInstructions->push_tail(Variable);
+	ParseState->symbols->add_variable(Variable);
+
+	ir_rvalue* VariableDeref = new(ParseState)ir_dereference_variable(Variable);
+
+	if (Frequency == HSF_HullShader && !OutputQualifier.Fields.bIsPatchConstant)
+	{
+		VariableDeref = new(ParseState)ir_dereference_array(VariableDeref, new(ParseState)ir_dereference_variable(ParseState->symbols->get_variable("gl_InvocationID")));
+	}
+
+	VariableDeref = new(ParseState)ir_dereference_record(VariableDeref, ralloc_strdup(ParseState, "Data"));
+
+	return VariableDeref;
+}
+
+/**
+* Generate an input semantic.
+* @param Frequency - The shader frequency.
+* @param ParseState - Parse state.
+* @param InputSemantic - The semantic name to generate.
+* @param InputQualifier - Qualifiers applied to the semantic.
+* @param InputVariableDeref - Deref for the argument variable.
+* @param DeclInstructions - IR to which declarations may be added.
+* @param PreCallInstructions - IR to which instructions may be added before the
+*                              entry point is called.
+*/
+static void GenShaderInputForVariable(
+	EHlslShaderFrequency Frequency,
+	_mesa_glsl_parse_state* ParseState,
+	const char* InputSemantic,
+	FSemanticQualifier InputQualifier,
+	ir_dereference* InputVariableDeref,
+	exec_list* DeclInstructions,
+	exec_list* PreCallInstructions,
+	int SemanticArraySize,
+	int SemanticArrayIndex
+	)
+{
+	const glsl_type* InputType = InputVariableDeref->type;
+
+	if (InputType->is_record())
+	{
+		for (uint32 i = 0; i < InputType->length; ++i)
+		{
+			const char* FieldSemantic = InputType->fields.structure[i].semantic;
+			const char* Semantic = 0;
+
+			if (InputSemantic && FieldSemantic)
+			{
+
+				_mesa_glsl_warning(ParseState, "semantic '%s' of field '%s' will be overridden by enclosing types' semantic '%s'",
+					InputType->fields.structure[i].semantic,
+					InputType->fields.structure[i].name,
+					InputSemantic);
+
+
+				FieldSemantic = 0;
+			}
+
+			if (InputSemantic && !FieldSemantic)
+			{
+				Semantic = ralloc_asprintf(ParseState, "%s%d", InputSemantic, i);
+				_mesa_glsl_warning(ParseState, "  creating semantic '%s' for struct field '%s'", Semantic, InputType->fields.structure[i].name);
+			}
+			else if (!InputSemantic && FieldSemantic)
+			{
+				Semantic = FieldSemantic;
+			}
+			else
+			{
+				Semantic = 0;
+			}
+
+			if (InputType->fields.structure[i].type->is_record() ||
+				Semantic)
+			{
+				FSemanticQualifier Qualifier = InputQualifier;
+				if (Qualifier.Packed == 0)
+				{
+					Qualifier.Fields.bCentroid = InputType->fields.structure[i].centroid;
+					Qualifier.Fields.InterpolationMode = InputType->fields.structure[i].interpolation;
+					Qualifier.Fields.bIsPatchConstant = InputType->fields.structure[i].patchconstant;
+				}
+
+				ir_dereference_record* FieldDeref = new(ParseState)ir_dereference_record(
+					InputVariableDeref->clone(ParseState, NULL),
+					InputType->fields.structure[i].name);
+				GenShaderInputForVariable(
+					Frequency,
+					ParseState,
+					Semantic,
+					Qualifier,
+					FieldDeref,
+					DeclInstructions,
+					PreCallInstructions,
+					SemanticArraySize,
+					SemanticArrayIndex
+					);
+			}
+			else
+			{
+				_mesa_glsl_error(
+					ParseState,
+					"field '%s' in input structure '%s' does not specify a semantic",
+					InputType->fields.structure[i].name,
+					InputType->name
+					);
+			}
+		}
+	}
+	else if (InputType->is_array() || InputType->is_inputpatch() || InputType->is_outputpatch())
+	{
+		int BaseIndex = 0;
+		const char* Semantic = 0;
+		check(InputSemantic);
+		ParseSemanticAndIndex(ParseState, InputSemantic, &Semantic, &BaseIndex);
+		check(BaseIndex >= 0);
+		check(InputType->is_array() || InputType->is_inputpatch() || InputType->is_outputpatch());
+		const unsigned ElementCount = InputType->is_array() ? InputType->length : InputType->patch_length;
+
+		{
+			//check(!InputQualifier.Fields.bIsPatchConstant);
+			InputQualifier.Fields.bIsPatchConstant = false;
+		}
+
+		for (unsigned i = 0; i < ElementCount; ++i)
+		{
+			ir_dereference_array* ArrayDeref = new(ParseState)ir_dereference_array(
+				InputVariableDeref->clone(ParseState, NULL),
+				new(ParseState)ir_constant((unsigned)i)
+				);
+			GenShaderInputForVariable(
+				Frequency,
+				ParseState,
+				ralloc_asprintf(ParseState, "%s%d", Semantic, BaseIndex + i),
+				InputQualifier,
+				ArrayDeref,
+				DeclInstructions,
+				PreCallInstructions,
+				SemanticArraySize,
+				SemanticArrayIndex
+				);
+		}
+	}
+	else
+	{
+		bool ApplyClipSpaceAdjustment = false;
+		ir_rvalue* SrcValue = GenShaderInputSemantic(Frequency, ParseState, InputSemantic,
+			InputQualifier, InputType, DeclInstructions, SemanticArraySize,
+			SemanticArrayIndex, ApplyClipSpaceAdjustment);
+
+		if (SrcValue)
+		{
+			YYLTYPE loc;
+
+			if (ParseState->adjust_clip_space_dx11_to_opengl && ApplyClipSpaceAdjustment)
+			{
+				// This is for input of gl_Position into geometry shader only.
+
+				// Generate a local variable to do the conversion in, keeping source type.
+				ir_variable* TempVariable = new(ParseState)ir_variable(SrcValue->type, NULL, ir_var_temporary);
+				PreCallInstructions->push_tail(TempVariable);
+
+				// Assign input to this variable
+				ir_dereference_variable* TempVariableDeref = new(ParseState)ir_dereference_variable(TempVariable);
+				PreCallInstructions->push_tail(
+					new(ParseState)ir_assignment(
+					TempVariableDeref,
+					SrcValue
+					)
+					);
+
+				// TempVariable.y = -TempVariable.y;
+				PreCallInstructions->push_tail(
+					new(ParseState)ir_assignment(
+					new(ParseState)ir_swizzle(TempVariableDeref->clone(ParseState, NULL), 1, 0, 0, 0, 1),
+					new(ParseState)ir_expression(ir_unop_neg,
+					glsl_type::float_type,
+					new(ParseState)ir_swizzle(TempVariableDeref->clone(ParseState, NULL), 1, 0, 0, 0, 1),
+					NULL)
+					)
+					);
+
+				// TempVariable.z = ( TempVariable.z + TempVariable.w ) / 2.0;
+				PreCallInstructions->push_tail(
+					new(ParseState)ir_assignment(
+					new(ParseState)ir_swizzle(TempVariableDeref->clone(ParseState, NULL), 2, 0, 0, 0, 1),
+					new(ParseState)ir_expression(ir_binop_div,
+					new(ParseState)ir_expression(ir_binop_add,
+					new(ParseState)ir_swizzle(TempVariableDeref->clone(ParseState, NULL), 2, 0, 0, 0, 1),
+					new(ParseState)ir_swizzle(TempVariableDeref->clone(ParseState, NULL), 3, 0, 0, 0, 1)
+					),
+					new(ParseState)ir_constant(2.0f)
+					)
+					)
+					);
+
+				SrcValue = TempVariableDeref->clone(ParseState, NULL);
+			}
+
+			apply_type_conversion(InputType, SrcValue, PreCallInstructions, ParseState, true, &loc);
+			PreCallInstructions->push_tail(
+				new(ParseState)ir_assignment(
+				InputVariableDeref->clone(ParseState, NULL),
+				SrcValue
+				)
+				);
+		}
+	}
+}
+
+
+/**
+* Generate a shader input.
+* @param Frequency - The shader frequency.
+* @param ParseState - Parse state.
+* @param InputSemantic - The semantic name to generate.
+* @param InputQualifier - Qualifiers applied to the semantic.
+* @param InputType - Value type.
+* @param DeclInstructions - IR to which declarations may be added.
+* @param PreCallInstructions - IR to which instructions may be added before the
+*                              entry point is called.
+* @returns the IR variable deref for the semantic.
+*/
+static ir_dereference_variable* GenShaderInput(
+	EHlslShaderFrequency Frequency,
+	_mesa_glsl_parse_state* ParseState,
+	const char* InputSemantic,
+	FSemanticQualifier InputQualifier,
+	const glsl_type* InputType,
+	exec_list* DeclInstructions,
+	exec_list* PreCallInstructions)
+{
+	ir_variable* TempVariable = new(ParseState)ir_variable(
+		InputType,
+		NULL,
+		ir_var_temporary);
+	ir_dereference_variable* TempVariableDeref = new(ParseState)ir_dereference_variable(TempVariable);
+	PreCallInstructions->push_tail(TempVariable);
+
+	//check ( InputSemantic ?  (_strnicmp(InputSemantic, "SV_", 3) ==0) : true);
+
+
+	// everything that's not an Outputpatch is patch constant. System values are treated specially
+	if (Frequency == HSF_DomainShader && !InputType->is_outputpatch())
+	{
+		InputQualifier.Fields.bIsPatchConstant = true;
+	}
+
+	if ((Frequency == HSF_GeometryShader && TempVariableDeref->type->is_array()) ||
+		(Frequency == HSF_HullShader && TempVariableDeref->type->is_inputpatch()) ||
+		(Frequency == HSF_DomainShader && TempVariableDeref->type->is_outputpatch())
+		)
+	{
+		check(InputType->is_array() || InputType->is_inputpatch() || InputType->is_outputpatch());
+		check(InputType->length || InputType->patch_length);
+
+
+		const unsigned ElementCount = InputType->is_array() ? InputType->length : InputType->patch_length;
+
+		if (!ParseState->symbols->get_variable("gl_in"))
+		{
+			// Create a built-in OpenGL variable gl_in[] containing built-in types.
+			// This variable will be used for OpenGL optimization by IR, so IR must know about it,
+			// but will not end up in final GLSL code.
+
+			// It has to be created here, as it contains multiple built-in variables in one interface block,
+			// which is not usual, so avoiding special cases in code.
+
+			glsl_struct_field *BuiltinFields = ralloc_array(ParseState, glsl_struct_field, 3);
+			memset(BuiltinFields, 0, 3 * sizeof(glsl_struct_field));
+
+			BuiltinFields[0].type = glsl_type::vec4_type;
+			BuiltinFields[0].name = ralloc_strdup(ParseState, "gl_Position");
+			BuiltinFields[1].type = glsl_type::float_type;
+			BuiltinFields[1].name = ralloc_strdup(ParseState, "gl_PointSize");
+			BuiltinFields[2].type = glsl_type::get_array_instance(glsl_type::float_type, 6);	// magic number is gl_MaxClipDistances
+			BuiltinFields[2].name = ralloc_strdup(ParseState, "gl_ClipDistance");
+
+			const glsl_type* BuiltinStruct = glsl_type::get_record_instance(BuiltinFields, 3, "gl_PerVertex");
+			const glsl_type* BuiltinArray = glsl_type::get_array_instance(BuiltinStruct, ElementCount);
+			ir_variable* BuiltinVariable = new(ParseState)ir_variable(BuiltinArray, "gl_in", ir_var_in);
+			BuiltinVariable->read_only = true;
+			BuiltinVariable->is_interface_block = true;
+			DeclInstructions->push_tail(BuiltinVariable);
+			ParseState->symbols->add_variable(BuiltinVariable);
+		}
+
+		for (unsigned i = 0; i < ElementCount; ++i)
+		{
+			ir_dereference_array* ArrayDeref = new(ParseState)ir_dereference_array(
+				TempVariableDeref->clone(ParseState, NULL),
+				new(ParseState)ir_constant((unsigned)i)
+				);
+			// Parse input variable
+			GenShaderInputForVariable(
+				Frequency,
+				ParseState,
+				InputSemantic,
+				InputQualifier,
+				ArrayDeref,
+				DeclInstructions,
+				PreCallInstructions,
+				ElementCount,
+				i
+				);
+		}
+	}
+	else
+	{
+		GenShaderInputForVariable(
+			Frequency,
+			ParseState,
+			InputSemantic,
+			InputQualifier,
+			TempVariableDeref,
+			DeclInstructions,
+			PreCallInstructions,
+			0,
+			0
+			);
+	}
+	return TempVariableDeref;
+}
+
+/**
+* Generate an output semantic.
+* @param Frequency - The shader frequency.
+* @param ParseState - Parse state.
+* @param OutputSemantic - The semantic name to generate.
+* @param OutputQualifier - Qualifiers applied to the semantic.
+* @param OutputVariableDeref - Deref for the argument variable.
+* @param DeclInstructions - IR to which declarations may be added.
+* @param PostCallInstructions - IR to which instructions may be added after the
+*                               entry point returns.
+*/
+void GenShaderOutputForVariable(
+	EHlslShaderFrequency Frequency,
+	_mesa_glsl_parse_state* ParseState,
+	const char* OutputSemantic,
+	FSemanticQualifier OutputQualifier,
+	ir_dereference* OutputVariableDeref,
+	exec_list* DeclInstructions,
+	exec_list* PostCallInstructions,
+	int SemanticArraySize,
+	int SemanticArrayIndex
+	)
+{
+	const glsl_type* OutputType = OutputVariableDeref->type;
+	if (OutputType->is_record())
+	{
+		for (uint32 i = 0; i < OutputType->length; ++i)
+		{
+			const char* FieldSemantic = OutputType->fields.structure[i].semantic;
+			const char* Semantic = 0;
+
+			if (OutputSemantic && FieldSemantic)
+			{
+
+				_mesa_glsl_warning(ParseState, "semantic '%s' of field '%s' will be overridden by enclosing types' semantic '%s'",
+					OutputType->fields.structure[i].semantic,
+					OutputType->fields.structure[i].name,
+					OutputSemantic);
+
+
+				FieldSemantic = 0;
+			}
+
+			if (OutputSemantic && !FieldSemantic)
+			{
+				Semantic = ralloc_asprintf(ParseState, "%s%d", OutputSemantic, i);
+				_mesa_glsl_warning(ParseState, "  creating semantic '%s' for struct field '%s'", Semantic, OutputType->fields.structure[i].name);
+			}
+			else if (!OutputSemantic && FieldSemantic)
+			{
+				Semantic = FieldSemantic;
+			}
+			else
+			{
+				Semantic = 0;
+			}
+
+			if (OutputType->fields.structure[i].type->is_record() ||
+				Semantic
+				)
+			{
+				FSemanticQualifier Qualifier = OutputQualifier;
+				if (Qualifier.Packed == 0)
+				{
+					Qualifier.Fields.bCentroid = OutputType->fields.structure[i].centroid;
+					Qualifier.Fields.InterpolationMode = OutputType->fields.structure[i].interpolation;
+					Qualifier.Fields.bIsPatchConstant = OutputType->fields.structure[i].patchconstant;
+				}
+
+				// Dereference the field and generate shader outputs for the field.
+				ir_dereference* FieldDeref = new(ParseState)ir_dereference_record(
+					OutputVariableDeref->clone(ParseState, NULL),
+					OutputType->fields.structure[i].name);
+				GenShaderOutputForVariable(
+					Frequency,
+					ParseState,
+					Semantic,
+					Qualifier,
+					FieldDeref,
+					DeclInstructions,
+					PostCallInstructions,
+					SemanticArraySize,
+					SemanticArrayIndex
+					);
+			}
+			else
+			{
+				_mesa_glsl_error(
+					ParseState,
+					"field '%s' in output structure '%s' does not specify a semantic",
+					OutputType->fields.structure[i].name,
+					OutputType->name
+					);
+			}
+		}
+	}
+	// TODO clean this up!!
+	else if ((OutputType->is_array() || OutputType->is_outputpatch()))
+	{
+		if (OutputSemantic)
+		{
+			int BaseIndex = 0;
+			const char* Semantic = 0;
+
+			ParseSemanticAndIndex(ParseState, OutputSemantic, &Semantic, &BaseIndex);
+
+			const unsigned ElementCount = OutputType->is_array() ? OutputType->length : (OutputType->patch_length);
+
+			for (unsigned i = 0; i < ElementCount; ++i)
+			{
+				ir_dereference_array* ArrayDeref = new(ParseState)ir_dereference_array(
+					OutputVariableDeref->clone(ParseState, NULL),
+					new(ParseState)ir_constant((unsigned)i)
+					);
+				GenShaderOutputForVariable(
+					Frequency,
+					ParseState,
+					ralloc_asprintf(ParseState, "%s%d", Semantic, BaseIndex + i),
+					OutputQualifier,
+					ArrayDeref,
+					DeclInstructions,
+					PostCallInstructions,
+					SemanticArraySize,
+					SemanticArrayIndex
+					);
+			}
+		}
+		else
+		{
+			_mesa_glsl_error(ParseState, "entry point does not specify a semantic for its return value");
+		}
+	}
+	else
+	{
+		if (OutputSemantic)
+		{
+			YYLTYPE loc;
+			ir_rvalue* Src = OutputVariableDeref->clone(ParseState, NULL);
+			const glsl_type* DestVariableType = NULL;
+			bool ApplyClipSpaceAdjustment = false;
+			bool ApplyClampPowerOfTwo = false;
+			ir_rvalue* DestVariableDeref = GenShaderOutputSemantic(Frequency, ParseState, OutputSemantic,
+				OutputQualifier, OutputType, DeclInstructions, &DestVariableType, ApplyClipSpaceAdjustment, ApplyClampPowerOfTwo);
+
+			apply_type_conversion(DestVariableType, Src, PostCallInstructions, ParseState, true, &loc);
+
+			if (ParseState->adjust_clip_space_dx11_to_opengl && ApplyClipSpaceAdjustment)
+			{
+				// Src.y = -Src.y;
+				PostCallInstructions->push_tail(
+					new(ParseState)ir_assignment(
+					new(ParseState)ir_swizzle(Src->clone(ParseState, NULL), 1, 0, 0, 0, 1),
+					new(ParseState)ir_expression(ir_unop_neg,
+					glsl_type::float_type,
+					new(ParseState)ir_swizzle(Src->clone(ParseState, NULL), 1, 0, 0, 0, 1),
+					NULL)
+					)
+					);
+
+				// Src.z = ( 2.0 * Src.z ) - Src.w;
+				PostCallInstructions->push_tail(
+					new(ParseState)ir_assignment(
+					new(ParseState)ir_swizzle(Src->clone(ParseState, NULL), 2, 0, 0, 0, 1),
+					new(ParseState)ir_expression(ir_binop_sub,
+					new(ParseState)ir_expression(ir_binop_mul,
+					new(ParseState)ir_constant(2.0f),
+					new(ParseState)ir_swizzle(Src->clone(ParseState, NULL), 2, 0, 0, 0, 1)
+					),
+					new(ParseState)ir_swizzle(Src->clone(ParseState, NULL), 3, 0, 0, 0, 1)
+					)
+					)
+					);
+			}
+
+			// GLSL doesn't support pow2 partitioning, so we treate pow2 as integer partitioning and
+			// manually compute the next power of two via exp2(pow(ceil(log2(Src)));
+			if (ApplyClampPowerOfTwo)
+			{
+				ir_variable* temp = new(ParseState)ir_variable(glsl_type::float_type, NULL, ir_var_temporary);
+
+				PostCallInstructions->push_tail(temp);
+
+				PostCallInstructions->push_tail(
+					new(ParseState)ir_assignment(
+					new(ParseState)ir_dereference_variable(temp),
+					new(ParseState)ir_expression(ir_unop_exp2,
+					new(ParseState)ir_expression(ir_unop_ceil,
+					new(ParseState)ir_expression(ir_unop_log2,
+					glsl_type::float_type,
+					Src->clone(ParseState, NULL),
+					NULL
+					)
+					)
+					)
+					)
+					);
+
+				// assign pow2 clamped variable to output variable
+				PostCallInstructions->push_tail(
+					new(ParseState)ir_assignment(
+					DestVariableDeref->clone(ParseState, NULL),
+					new(ParseState)ir_dereference_variable(temp)
+					)
+					);
+			}
+			else
+			{
+				PostCallInstructions->push_tail(new(ParseState)ir_assignment(DestVariableDeref, Src));
+			}
+		}
+		else
+		{
+			_mesa_glsl_error(ParseState, "entry point does not specify a semantic for its return value");
+		}
+	}
+}
+/**
+* Generate an output semantic.
+* @param Frequency - The shader frequency.
+* @param ParseState - Parse state.
+* @param OutputSemantic - The semantic name to generate.
+* @param OutputQualifier - Qualifiers applied to the semantic.
+* @param OutputType - Value type.
+* @param DeclInstructions - IR to which declarations may be added.
+* @param PreCallInstructions - IR to which isntructions may be added before the
+entry point is called.
+* @param PostCallInstructions - IR to which instructions may be added after the
+*                               entry point returns.
+* @returns the IR variable deref for the semantic.
+*/
+static ir_dereference_variable* GenShaderOutput(
+	EHlslShaderFrequency Frequency,
+	_mesa_glsl_parse_state* ParseState,
+	const char* OutputSemantic,
+	FSemanticQualifier OutputQualifier,
+	const glsl_type* OutputType,
+	exec_list* DeclInstructions,
+	exec_list* PreCallInstructions,
+	exec_list* PostCallInstructions
+	)
+{
+	// Generate a local variable to hold the output.
+	ir_variable* TempVariable = new(ParseState)ir_variable(
+		OutputType,
+		NULL,
+		ir_var_temporary);
+	ir_dereference_variable* TempVariableDeref = new(ParseState)ir_dereference_variable(TempVariable);
+	PreCallInstructions->push_tail(TempVariable);
+	GenShaderOutputForVariable(
+		Frequency,
+		ParseState,
+		OutputSemantic,
+		OutputQualifier,
+		TempVariableDeref,
+		DeclInstructions,
+		PostCallInstructions,
+		0,
+		0
+		);
+	return TempVariableDeref;
+}
+
+static void GenerateAppendFunctionBody(
+	_mesa_glsl_parse_state* ParseState,
+	exec_list* DeclInstructions,
+	const glsl_type* geometry_append_type
+	)
+{
+	ir_function *func = ParseState->symbols->get_function("OutputStream_Append");
+	check(func);
+
+	exec_list comparison_parameter;
+	ir_variable* var = new(ParseState)ir_variable(geometry_append_type, ralloc_asprintf(ParseState, "arg0"), ir_var_in);
+	comparison_parameter.push_tail(var);
+
+	bool is_exact = false;
+	ir_function_signature *sig = func->matching_signature(&comparison_parameter, &is_exact);
+	check(sig && is_exact);
+	var = (ir_variable*)sig->parameters.get_head();
+
+	//	{
+	//		const glsl_type* output_type = var->type;
+	//		_mesa_glsl_warning(ParseState, "GenerateAppendFunctionBody: parsing argument struct '%s'", output_type->name );
+	//		int indexof_RenderTargetArrayIndex = -1;
+	//		for (int i = 0; i < output_type->length; i++)
+	//		{
+	//			_mesa_glsl_warning(ParseState, "   name '%s' : semantic '%s'", output_type->fields.structure[i].name, output_type->fields.structure[i].semantic );
+	//		}
+	//	}
+
+	// Generate assignment instructions from function argument to out variables
+	FSemanticQualifier OutputQualifier;
+	ir_dereference_variable* TempVariableDeref = new(ParseState)ir_dereference_variable(var);
+	GenShaderOutputForVariable(
+		HSF_GeometryShader,
+		ParseState,
+		NULL,
+		OutputQualifier,
+		TempVariableDeref,
+		DeclInstructions,
+		&sig->body,
+		0,
+		0
+		);
+
+	// If the output structure type contains a SV_RenderTargetArrayIndex semantic, add a custom user output semantic.
+	// It's used to pass layer index to pixel shader, as GLSL 1.50 doesn't allow pixel shader to read from gl_Layer.
+	const glsl_type* output_type = var->type;
+	int indexof_RenderTargetArrayIndex = -1;
+	for (uint32 i = 0; i < output_type->length; i++)
+	{
+		if (output_type->fields.structure[i].semantic && (strcmp(output_type->fields.structure[i].semantic, "SV_RenderTargetArrayIndex") == 0))
+		{
+			indexof_RenderTargetArrayIndex = i;
+			break;
+		}
+	}
+
+	if (indexof_RenderTargetArrayIndex != -1)
+	{
+		// Add the new member with semantic
+		glsl_struct_field field;
+		field.type = output_type->fields.structure[indexof_RenderTargetArrayIndex].type;
+		field.name = "HLSLCCLayerIndex";
+		field.semantic = CUSTOM_LAYER_INDEX_SEMANTIC;
+		field.centroid = 0;
+		field.interpolation = ir_interp_qualifier_flat;
+		field.geometryinput = 0;
+		field.patchconstant = 0;
+
+		glsl_type* non_const_type = (glsl_type*)output_type;
+		non_const_type->add_structure_member(&field);
+
+		// Create new out variable for the new member and generate assignment that will copy input's layer index field to it
+		FSemanticQualifier Qualifier;
+		Qualifier.Fields.bCentroid = 0;
+		Qualifier.Fields.InterpolationMode = ir_interp_qualifier_flat;
+
+		const glsl_type* new_output_type = ((ir_variable*)sig->parameters.get_head())->type;
+		GenShaderOutputForVariable(
+			HSF_GeometryShader,
+			ParseState,
+			CUSTOM_LAYER_INDEX_SEMANTIC,
+			Qualifier,
+			new(ParseState)ir_dereference_record(var, new_output_type->fields.structure[indexof_RenderTargetArrayIndex].name),
+			DeclInstructions,
+			&sig->body,
+			0,
+			0
+			);
+	}
+
+	// Call EmitVertex()
+	ir_function *emitVertexFunc = ParseState->symbols->get_function("EmitVertex");
+	check(emitVertexFunc);
+	check(emitVertexFunc->signatures.get_head() == emitVertexFunc->signatures.get_tail());
+	ir_function_signature *emitVertexSig = (ir_function_signature *)emitVertexFunc->signatures.get_head();
+	exec_list actual_parameter;
+	sig->body.push_tail(new(ParseState)ir_call(emitVertexSig, NULL, &actual_parameter));
+}
+
+bool FGlslCodeBackend::GenerateMain(
+	EHlslShaderFrequency Frequency,
+	const char* EntryPoint,
+	exec_list* Instructions,
+	_mesa_glsl_parse_state* ParseState)
+{
+	{
+		// Set up origin_upper_left for gl_FragCoord, depending on HLSLCC_DX11ClipSpace flag presence.
+		FSystemValue* SystemValues = SystemValueTable[HSF_PixelShader];
+		for (int i = 0; SystemValues[i].Semantic != NULL; ++i)
+		{
+			if (_stricmp(SystemValues[i].GlslName, "gl_FragCoord") == 0)
+			{
+				SystemValues[i].bOriginUpperLeft = !ParseState->adjust_clip_space_dx11_to_opengl;
+				break;
+			}
+		}
+	}
+
+	ir_function_signature* EntryPointSig = FindEntryPointFunction(Instructions, ParseState, EntryPoint);
+	if (EntryPointSig)
+	{
+		void* TempMemContext = ralloc_context(NULL);
+		exec_list DeclInstructions;
+		exec_list PreCallInstructions;
+		exec_list ArgInstructions;
+		exec_list PostCallInstructions;
+		const glsl_type* geometry_append_type = NULL;
+
+		ParseState->maxvertexcount = EntryPointSig->maxvertexcount;
+
+		ParseState->tessellation = EntryPointSig->tessellation;
+
+		ParseState->symbols->push_scope();
+
+		foreach_iter(exec_list_iterator, Iter, EntryPointSig->parameters)
+		{
+			ir_variable *const Variable = (ir_variable *)Iter.get();
+			if (Variable->semantic != NULL || Variable->type->is_record()
+				|| (Frequency == HSF_GeometryShader && (Variable->type->is_outputstream() || Variable->type->is_array()))
+				|| (Frequency == HSF_HullShader && (Variable->type->is_patch()))
+				|| (Frequency == HSF_DomainShader && (Variable->type->is_outputpatch()))
+				)
+			{
+				FSemanticQualifier Qualifier;
+				Qualifier.Fields.bCentroid = Variable->centroid;
+				Variable->centroid = 0;
+				Qualifier.Fields.InterpolationMode = Variable->interpolation;
+				Variable->interpolation = 0;
+				Qualifier.Fields.bIsPatchConstant = Variable->is_patch_constant;
+				Variable->is_patch_constant = 0;
+
+				ir_dereference_variable* ArgVarDeref = NULL;
+				switch (Variable->mode)
+				{
+				case ir_var_in:
+					if (Frequency == HSF_GeometryShader && Variable->type->is_array())
+					{
+						// Remember information about geometry input type globally
+						ParseState->geometryinput = Variable->geometryinput;
+					}
+
+					if (Frequency == HSF_PixelShader)
+					{
+						// Replace SV_RenderTargetArrayIndex in
+						// input structure semantic with custom semantic.
+						if (Variable->semantic && (strcmp(Variable->semantic, "SV_RenderTargetArrayIndex") == 0))
+						{
+							//							_mesa_glsl_warning(ParseState, "Replacing semantic of variable '%s' with our custom one", Variable->name);
+							Variable->semantic = ralloc_strdup(Variable, CUSTOM_LAYER_INDEX_SEMANTIC);
+							Variable->interpolation = ir_interp_qualifier_flat;
+						}
+						else if (Variable->type->is_record())
+						{
+							const glsl_type* output_type = Variable->type;
+							int indexof_RenderTargetArrayIndex = -1;
+							for (uint32 i = 0; i < output_type->length; i++)
+							{
+								if (Variable->type->fields.structure[i].semantic && (strcmp(Variable->type->fields.structure[i].semantic, "SV_RenderTargetArrayIndex") == 0))
+								{
+									indexof_RenderTargetArrayIndex = i;
+									break;
+								}
+							}
+
+							if (indexof_RenderTargetArrayIndex != -1)
+							{
+								// _mesa_glsl_warning(ParseState, "Replacing semantic of member %d of variable '%s' with our custom one", indexof_RenderTargetArrayIndex, Variable->name);
+								// Replace the member with one with semantic
+								glsl_struct_field field;
+								field.type = Variable->type->fields.structure[indexof_RenderTargetArrayIndex].type;
+								field.name = Variable->type->fields.structure[indexof_RenderTargetArrayIndex].name;
+								field.semantic = CUSTOM_LAYER_INDEX_SEMANTIC;
+								field.centroid = 0;
+								field.interpolation = ir_interp_qualifier_flat;
+								field.geometryinput = 0;
+								field.patchconstant = 0;
+
+								glsl_type* non_const_type = (glsl_type*)output_type;
+								non_const_type->replace_structure_member(indexof_RenderTargetArrayIndex, &field);
+							}
+						}
+					}
+
+					ArgVarDeref = GenShaderInput(
+						Frequency,
+						ParseState,
+						Variable->semantic,
+						Qualifier,
+						Variable->type,
+						&DeclInstructions,
+						&PreCallInstructions
+						);
+					break;
+				case ir_var_out:
+					ArgVarDeref = GenShaderOutput(
+						Frequency,
+						ParseState,
+						Variable->semantic,
+						Qualifier,
+						Variable->type,
+						&DeclInstructions,
+						&PreCallInstructions,
+						&PostCallInstructions
+						);
+					break;
+				case ir_var_inout:
+				{
+									 check(Frequency == HSF_GeometryShader);
+									 // This is an output stream for geometry shader. It's not referenced as a variable inside the function,
+									 // instead OutputStream.Append(vertex) and OutputStream.RestartStrip() are called, and this variable
+									 // has already been optimized out of them in ast_to_hir translation.
+
+									 // Generate a local variable to add to arguments. It won't be referenced anywhere, so it should get optimized out.
+									 ir_variable* TempVariable = new(ParseState)ir_variable(
+										 Variable->type,
+										 NULL,
+										 ir_var_temporary);
+									 ArgVarDeref = new(ParseState)ir_dereference_variable(TempVariable);
+									 PreCallInstructions.push_tail(TempVariable);
+
+									 // We need to move this information somewhere safer, as this pseudo-variable will get optimized out of existence
+									 ParseState->outputstream_type = Variable->type->outputstream_type;
+
+									 check(Variable->type->is_outputstream());
+									 check(Variable->type->inner_type->is_record());
+
+									 geometry_append_type = Variable->type->inner_type;
+				}
+					break;
+				default:
+				{
+						   _mesa_glsl_error(
+							   ParseState,
+							   "entry point parameter '%s' must be an input or output",
+							   Variable->name
+							   );
+				}
+				}
+				ArgInstructions.push_tail(ArgVarDeref);
+			}
+			else
+			{
+				_mesa_glsl_error(ParseState, "entry point parameter "
+					"'%s' does not specify a semantic", Variable->name);
+			}
+		}
+
+		// The function's return value should have an output semantic if it's not void.
+		ir_dereference_variable* EntryPointReturn = NULL;
+		if (EntryPointSig->return_type->is_void() == false)
+		{
+			FSemanticQualifier Qualifier;
+			EntryPointReturn = GenShaderOutput(
+				Frequency,
+				ParseState,
+				EntryPointSig->return_semantic,
+				Qualifier,
+				EntryPointSig->return_type,
+				&DeclInstructions,
+				&PreCallInstructions,
+				&PostCallInstructions
+				);
+		}
+
+		if (Frequency == HSF_GeometryShader)
+		{
+			GenerateAppendFunctionBody(
+				ParseState,
+				&DeclInstructions,
+				geometry_append_type
+				);
+		}
+
+		/*
+		we map the HLSL hull shader to this GLSL main function
+		for the most parts, we treat variables of InputPatch and OutputPatch as arrays of the inner type
+
+		build input patch from shader input interface blocks
+		call hull shader main function with input patch and current control point id (gl_InvocationID)
+		copy hull shader main result for the current control point to the proper shader output interface block element
+		barrier
+		(so all instances have computed the per control point data)
+		build patch constant function input (of type output patch) from the shader output interface blocks
+		(need to do this, since this is the only shader variable shared between control points running in parallel)
+
+		if control point id (gl_InvocationID) is 0
+		call patch constant function with the output patch as an input
+		copy the patch constant result to the "patch" shader output interface block
+		*/
+
+		if (Frequency == HSF_HullShader)
+		{
+
+			ir_function_signature* PatchConstantSig = FindPatchConstantFunction(Instructions, ParseState);
+
+			if (!PatchConstantSig)
+			{
+				_mesa_glsl_error(ParseState, "patch constant function `%s' not found", ParseState->tessellation.patchconstantfunc);
+			}
+
+			const glsl_type* OutputPatchType = glsl_type::get_templated_instance(EntryPointReturn->type, "OutputPatch", 0, ParseState->tessellation.outputcontrolpoints);
+
+			ir_variable* OutputPatchVar = new(ParseState)ir_variable(OutputPatchType, NULL, ir_var_temporary);
+
+
+			// call barrier() to ensure that all threads have computed the per-patch computation
+			{
+				// We can't just use the symbol table b/c it only has the HLSL and not the GLSL barrier functions in it
+				foreach_iter(exec_list_iterator, Iter, *Instructions)
+				{
+					ir_instruction *ir = (ir_instruction *)Iter.get();
+					ir_function *Function = ir->as_function();
+					if (Function && strcmp(Function->name, "barrier") == 0)
+					{
+						check(Function->signatures.get_head() == Function->signatures.get_tail());
+						exec_list VoidParameter;
+						ir_function_signature * BarrierFunctionSig = Function->matching_signature(&VoidParameter);
+						PostCallInstructions.push_tail(new(ParseState)ir_call(BarrierFunctionSig, NULL, &VoidParameter));
+					}
+				}
+			}
+
+			// reassemble output patch variable(for the patch constant function) from the shader outputs
+			GenShaderPatchConstantFunctionInputs(ParseState, OutputPatchVar, PostCallInstructions);
+
+			// call the entry point
+			if (PatchConstantSig)
+			{
+				CallPatchConstantFunction(ParseState, OutputPatchVar, PatchConstantSig, DeclInstructions, PostCallInstructions);
+			}
+		}
+
+		ParseState->symbols->pop_scope();
+
+		// Build the void main() function for GLSL.
+		ir_function_signature* MainSig = new(ParseState)ir_function_signature(glsl_type::void_type);
+		MainSig->is_defined = true;
+		MainSig->is_main = true;
+		MainSig->body.append_list(&PreCallInstructions);
+		MainSig->body.push_tail(new(ParseState)ir_call(EntryPointSig, EntryPointReturn, &ArgInstructions));
+		MainSig->body.append_list(&PostCallInstructions);
+		MainSig->maxvertexcount = EntryPointSig->maxvertexcount;
+		MainSig->is_early_depth_stencil = EntryPointSig->is_early_depth_stencil;
+		MainSig->wg_size_x = EntryPointSig->wg_size_x;
+		MainSig->wg_size_y = EntryPointSig->wg_size_y;
+		MainSig->wg_size_z = EntryPointSig->wg_size_z;
+		MainSig->tessellation = EntryPointSig->tessellation;
+
+		if (MainSig->is_early_depth_stencil && Frequency != HSF_PixelShader)
+		{
+			_mesa_glsl_error(ParseState, "'earlydepthstencil' attribute only applies to pixel shaders");
+		}
+
+		if (MainSig->maxvertexcount > 0 && Frequency != HSF_GeometryShader)
+		{
+			_mesa_glsl_error(ParseState, "'maxvertexcount' attribute only applies to geometry shaders");
+		}
+
+		if (MainSig->is_early_depth_stencil && ParseState->language_version < 310)
+		{
+			_mesa_glsl_error(ParseState, "'earlydepthstencil' attribute only supported on GLSL 4.30 target and later");
+		}
+
+		if (MainSig->wg_size_x > 0 && Frequency != HSF_ComputeShader)
+		{
+			_mesa_glsl_error(ParseState, "'num_threads' attribute only applies to compute shaders");
+		}
+
+		// in GLSL, unlike in HLSL fixed-function tessellator properties are specified on the domain shader
+		// and not the hull shader, so we specify them for both in the .usf shaders and then print a warning,
+		// similar to what fxc is doing
+
+		if (MainSig->tessellation.domain != GLSL_DOMAIN_NONE && (Frequency != HSF_HullShader && Frequency != HSF_DomainShader))
+		{
+			_mesa_glsl_warning(ParseState, "'domain' attribute only applies to hull or domain shaders");
+		}
+
+		if (MainSig->tessellation.outputtopology != GLSL_OUTPUTTOPOLOGY_NONE && Frequency != HSF_HullShader)
+		{
+			_mesa_glsl_warning(ParseState, "'outputtopology' attribute only applies to hull shaders");
+		}
+
+		if (MainSig->tessellation.partitioning != GLSL_PARTITIONING_NONE && Frequency != HSF_HullShader)
+		{
+			_mesa_glsl_warning(ParseState, "'partitioning' attribute only applies to hull shaders");
+		}
+
+		if (MainSig->tessellation.outputcontrolpoints > 0 && Frequency != HSF_HullShader)
+		{
+			_mesa_glsl_warning(ParseState, "'outputcontrolpoints' attribute only applies to hull shaders");
+		}
+
+		if (MainSig->tessellation.maxtessfactor > 0.0f && Frequency != HSF_HullShader)
+		{
+			_mesa_glsl_warning(ParseState, "'maxtessfactor' attribute only applies to hull shaders");
+		}
+
+		if (MainSig->tessellation.patchconstantfunc != 0 && Frequency != HSF_HullShader)
+		{
+			_mesa_glsl_warning(ParseState, "'patchconstantfunc' attribute only applies to hull shaders");
+		}
+
+		ir_function* MainFunction = new(ParseState)ir_function("main");
+		MainFunction->add_signature(MainSig);
+
+		Instructions->append_list(&DeclInstructions);
+		Instructions->push_tail(MainFunction);
+
+		ralloc_free(TempMemContext);
+
+		// Now that we have a proper main(), move global setup to main().
+		MoveGlobalInstructionsToMain(Instructions);
+	}
+	else
+	{
+		_mesa_glsl_error(ParseState, "shader entry point '%s' not "
+			"found", EntryPoint);
+	}
+
+	return true;
+}
+ir_function_signature*  FGlslCodeBackend::FindPatchConstantFunction(exec_list* Instructions, _mesa_glsl_parse_state* ParseState)
+{
+	ir_function_signature* PatchConstantSig = 0;
+
+	// TODO refactor this and the fetching of the main siganture
+	foreach_iter(exec_list_iterator, Iter, *Instructions)
+	{
+		ir_instruction *ir = (ir_instruction *)Iter.get();
+		ir_function *Function = ir->as_function();
+		if (Function && strcmp(Function->name, ParseState->tessellation.patchconstantfunc) == 0)
+		{
+			int NumSigs = 0;
+			foreach_iter(exec_list_iterator, SigIter, *Function)
+			{
+				if (++NumSigs == 1)
+				{
+					PatchConstantSig = (ir_function_signature *)SigIter.get();
+				}
+			}
+			if (NumSigs == 1)
+			{
+				break;
+			}
+			else
+			{
+				_mesa_glsl_error(ParseState, "patch constant function "
+					"`%s' has multiple signatures", ParseState->tessellation.patchconstantfunc);
+			}
+		}
+	}
+
+	return PatchConstantSig;
+}
+
+void FGlslCodeBackend::CallPatchConstantFunction(_mesa_glsl_parse_state* ParseState, ir_variable* OutputPatchVar, ir_function_signature* PatchConstantSig, exec_list& DeclInstructions, exec_list &PostCallInstructions)
+{
+	exec_list PatchConstantArgs;
+	if (OutputPatchVar && !PatchConstantSig->parameters.is_empty())
+	{
+		PatchConstantArgs.push_tail(new(ParseState)ir_dereference_variable(OutputPatchVar));
+	}
+
+	ir_if* thread_if = new(ParseState)ir_if(
+		new(ParseState)ir_expression(
+		ir_binop_equal,
+		new (ParseState)ir_constant(
+		0
+		),
+		new (ParseState)ir_dereference_variable(
+		ParseState->symbols->get_variable("gl_InvocationID")
+		)
+		)
+		);
+
+	exec_list PrePatchConstCallInstructions;
+	exec_list PostPatchConstCallInstructions;
+
+	FSemanticQualifier Qualifier;
+	Qualifier.Fields.bIsPatchConstant = 1;
+
+	ir_dereference_variable* PatchConstantReturn = GenShaderOutput(
+		HSF_HullShader,
+		ParseState,
+		PatchConstantSig->return_semantic,
+		Qualifier,
+		PatchConstantSig->return_type,
+		&DeclInstructions,
+		&PrePatchConstCallInstructions,
+		&PostPatchConstCallInstructions
+		);
+
+	thread_if->then_instructions.append_list(&PrePatchConstCallInstructions);
+	thread_if->then_instructions.push_tail(new(ParseState)ir_call(PatchConstantSig, PatchConstantReturn, &PatchConstantArgs));
+	thread_if->then_instructions.append_list(&PostPatchConstCallInstructions);
+
+	PostCallInstructions.push_tail(thread_if);
+			}
+
+
+			/*
+reassemble output patch variable (for the patch constant function) from the shader outputs
+			turn this: (from the GenOutputs of calling the entry point main)
+
+			out_InnerMember[gl_InvocationID].Data = t2.Middle.Inner.Value;
+
+			into this:
+
+			//output_patch<FPNTessellationHSToDS> FPNTessellationHSToDS t3[3]; //output_patch<FPNTessellationHSToDS> ;
+			t3[0].Middle.Inner.Value = out_InnerMember[0].Data;
+			t3[1].Middle.Inner.Value = out_InnerMember[1].Data;
+			t3[2].Middle.Inner.Value = out_InnerMember[2].Data;
+
+			*/
+
+void FGlslCodeBackend::GenShaderPatchConstantFunctionInputs(_mesa_glsl_parse_state* ParseState, ir_variable* OutputPatchVar, exec_list &PostCallInstructions)
+{
+	PostCallInstructions.push_tail(OutputPatchVar);
+	foreach_iter(exec_list_iterator, Iter, PostCallInstructions)
+	{
+		ir_instruction *ir = (ir_instruction *)Iter.get();
+
+		ir_assignment* assignment = ir->as_assignment();
+
+		if (!assignment)
+		{
+			continue;
+		}
+
+		ir_dereference_record* lhs = assignment->lhs->as_dereference_record();
+		ir_rvalue* rhs = assignment->rhs;
+
+		if (!lhs)
+		{
+			continue;
+		}
+
+		if (!rhs)
+		{
+			continue;
+		}
+
+		ir_dereference_array* lhs_array = lhs->record->as_dereference_array();
+
+		if (!lhs_array)
+		{
+			continue;
+		}
+
+		ir_dereference_variable* OutputPatchArrayIndex = lhs_array->array_index->as_dereference_variable();
+		ir_dereference_variable* OutputPatchArray = lhs_array->array->as_dereference_variable();
+
+		if (!OutputPatchArrayIndex)
+		{
+			continue;
+		}
+
+		if (0 != strcmp(OutputPatchArrayIndex->var->name, "gl_InvocationID"))
+		{
+			continue;
+		}
+
+		if (!OutputPatchArray)
+		{
+			continue;
+		}
+
+		const char* OutArrayFieldName = lhs->field;
+
+		for (int OutputVertex = 0; OutputVertex < ParseState->tessellation.outputcontrolpoints; ++OutputVertex)
+		{
+			struct Helper
+			{
+				// the struct inside the output patch can have the actual outputs with semantics nested inside,
+
+				static void ReplaceVariableDerefWithArrayDeref(ir_instruction* Node, ir_dereference_array* ArrayDereference)
+				{
+					if (ir_dereference_record* AsRecord = Node->as_dereference_record())
+					{
+						if (AsRecord->record->as_dereference_variable())
+						{
+							AsRecord->record = ArrayDereference;
+						}
+						else
+						{
+							ReplaceVariableDerefWithArrayDeref(AsRecord->record, ArrayDereference);
+						}
+					}
+					else if (ir_dereference_array* AsArray = Node->as_dereference_array())
+					{
+						if (AsArray->array->as_dereference_variable())
+						{
+							AsArray->array = ArrayDereference;
+						}
+						else
+						{
+							ReplaceVariableDerefWithArrayDeref(AsArray->array, ArrayDereference);
+						}
+					}
+					else
+					{
+						check(false);
+					}
+				}
+			};
+
+			ir_dereference_array* OutputPatchElementIndex = new(ParseState)ir_dereference_array(
+				OutputPatchVar,
+				new(ParseState)ir_constant(
+				OutputVertex
+				)
+				);
+
+			ir_rvalue* OutputPatchElement = rhs->clone(ParseState, 0);
+			Helper::ReplaceVariableDerefWithArrayDeref(OutputPatchElement, OutputPatchElementIndex);
+
+			PostCallInstructions.push_tail(
+				new (ParseState)ir_assignment(
+					OutputPatchElement,
+					new(ParseState)ir_dereference_record(
+						new(ParseState)ir_dereference_array(
+							OutputPatchArray->clone(ParseState, 0),
+							new(ParseState)ir_constant(OutputVertex)
+						),
+						OutArrayFieldName
+					)
+				)
+				);
+		}
+	}
+}
+
+void FGlslLanguageSpec::SetupLanguageIntrinsics(_mesa_glsl_parse_state* State, exec_list* ir)
+{
+	if (bIsES2)
+	{
+		make_intrinsic_genType(ir, State, FRAMEBUFFER_FETCH_ES2, ir_invalid_opcode, IR_INTRINSIC_FLOAT, 0, 4, 4);
+	}
+}
+
diff --git a/Engine/Source/Developer/ShaderFormatOpenGL/Private/ir_gen_glsl.h b/Engine/Source/Developer/ShaderFormatOpenGL/Private/ir_gen_glsl.h
new file mode 100644
--- /dev/null
+++ b/Engine/Source/Developer/ShaderFormatOpenGL/Private/ir_gen_glsl.h
@@ -0,0 +1,72 @@
+// Copyright 1998-2014 Epic Games, Inc. All Rights Reserved.
+
+#pragma once
+
+#include "hlslcc.h"
+#include "LanguageSpec.h"
+
+class FGlslLanguageSpec : public ILanguageSpec
+{
+protected:
+	bool bIsES2;
+
+public:
+	FGlslLanguageSpec(bool bInIsES2) : bIsES2(bInIsES2) {}
+
+	virtual bool SupportsDeterminantIntrinsic() const override
+	{
+		return !bIsES2;
+	}
+
+	virtual bool SupportsTransposeIntrinsic() const override
+	{
+		return !bIsES2;
+	}
+
+	virtual bool SupportsIntegerModulo() const override
+	{
+		return !bIsES2;
+	}
+
+	virtual bool SupportsMatrixConversions() const override { return true; }
+
+	virtual void SetupLanguageIntrinsics(_mesa_glsl_parse_state* State, exec_list* ir) override;
+};
+
+class ir_variable;
+
+// Generates GLSL compliant code from IR tokens
+#ifdef __GNUC__
+#pragma GCC visibility push(default)
+#endif // __GNUC__
+struct FGlslCodeBackend : public FCodeBackend
+{
+	FGlslCodeBackend(unsigned int InHlslCompileFlags) :
+		FCodeBackend(InHlslCompileFlags)
+	{
+	}
+
+	virtual char* GenerateCode(struct exec_list* ir, struct _mesa_glsl_parse_state* ParseState, EHlslShaderFrequency Frequency) override;
+
+	// Return false if there were restrictions that made compilation fail
+	virtual bool ApplyAndVerifyPlatformRestrictions(exec_list* Instructions, _mesa_glsl_parse_state* ParseState, EHlslShaderFrequency Frequency) override;
+
+	/**
+	* Generate a GLSL main() function that calls the entry point and handles
+	* reading and writing all input and output semantics.
+	* @param Frequency - The shader frequency.
+	* @param EntryPoint - The name of the shader entry point.
+	* @param Instructions - IR code.
+	* @param ParseState - Parse state.
+	*/
+	virtual bool GenerateMain(EHlslShaderFrequency Frequency, const char* EntryPoint, exec_list* Instructions, _mesa_glsl_parse_state* ParseState) override;
+
+	void GenShaderPatchConstantFunctionInputs(_mesa_glsl_parse_state* ParseState, ir_variable* OutputPatchVar, exec_list &PostCallInstructions);
+
+	void CallPatchConstantFunction(_mesa_glsl_parse_state* ParseState, ir_variable* OutputPatchVar, ir_function_signature* PatchConstantSig, exec_list& DeclInstructions, exec_list &PostCallInstructions);
+
+	ir_function_signature* FindPatchConstantFunction(exec_list* Instructions, _mesa_glsl_parse_state* ParseState);
+};
+#ifdef __GNUC__
+#pragma GCC visibility pop
+#endif // __GNUC__
diff --git a/Engine/Source/Developer/iOS/MetalShaderFormat/Private/MetalBackend.cpp b/Engine/Source/Developer/iOS/MetalShaderFormat/Private/MetalBackend.cpp
--- a/Engine/Source/Developer/iOS/MetalShaderFormat/Private/MetalBackend.cpp
+++ b/Engine/Source/Developer/iOS/MetalShaderFormat/Private/MetalBackend.cpp
@@ -5,7 +5,6 @@
 #include "hlslcc.h"
 #include "hlslcc_private.h"
 #include "MetalBackend.h"
-#include "glsl/ir_gen_glsl.h"
 #include "mesa/compiler.h"
 #include "mesa/glsl_parser_extras.h"
 #include "mesa/hash_table.h"
diff --git a/Engine/Source/Developer/iOS/MetalShaderFormat/Private/MetalShaderCompiler.cpp b/Engine/Source/Developer/iOS/MetalShaderFormat/Private/MetalShaderCompiler.cpp
--- a/Engine/Source/Developer/iOS/MetalShaderFormat/Private/MetalShaderCompiler.cpp
+++ b/Engine/Source/Developer/iOS/MetalShaderFormat/Private/MetalShaderCompiler.cpp
@@ -811,7 +811,7 @@ static FString CreateCommandLineHLSLCC( const FString& ShaderFile, const FString
 	const TCHAR* VersionSwitch = TEXT("-metal");
 	const TCHAR* FlattenUB = ((CCFlags & HLSLCC_FlattenUniformBuffers) == HLSLCC_FlattenUniformBuffers) ? TEXT("-flattenub") : TEXT("");
 
-	return CreateCrossCompilerBatchFileContents(ShaderFile, OutputFile, FrequencySwitch, EntryPoint, VersionSwitch, true, FlattenUB);
+	return CreateCrossCompilerBatchFileContents(ShaderFile, OutputFile, FrequencySwitch, EntryPoint, VersionSwitch, FlattenUB);
 }
 
 void CompileShader_Metal(const FShaderCompilerInput& Input,FShaderCompilerOutput& Output,const FString& WorkingDirectory)
diff --git a/Engine/Source/Developer/iOS/MetalShaderFormat/Private/MetalUtils.cpp b/Engine/Source/Developer/iOS/MetalShaderFormat/Private/MetalUtils.cpp
--- a/Engine/Source/Developer/iOS/MetalShaderFormat/Private/MetalUtils.cpp
+++ b/Engine/Source/Developer/iOS/MetalShaderFormat/Private/MetalUtils.cpp
@@ -5,7 +5,6 @@
 #include "hlslcc.h"
 #include "hlslcc_private.h"
 #include "MetalBackend.h"
-#include "glsl/ir_gen_glsl.h"
 #include "mesa/compiler.h"
 #include "mesa/glsl_parser_extras.h"
 #include "mesa/hash_table.h"
diff --git a/Engine/Source/Programs/CrossCompilerTool/CrossCompilerTool.Build.cs b/Engine/Source/Programs/CrossCompilerTool/CrossCompilerTool.Build.cs
--- a/Engine/Source/Programs/CrossCompilerTool/CrossCompilerTool.Build.cs
+++ b/Engine/Source/Programs/CrossCompilerTool/CrossCompilerTool.Build.cs
@@ -10,12 +10,14 @@ public class CrossCompilerTool : ModuleRules
 
 		PrivateIncludePaths.Add("Runtime/Launch/Private");		// For LaunchEngineLoop.cpp include
 		PrivateIncludePaths.Add("Developer/iOS/MetalShaderFormat/Private");		// For Metal includes
+		PrivateIncludePaths.Add("Developer/ShaderFormatOpenGL/Private");		// For GLSL includes
 
 		PrivateDependencyModuleNames.AddRange(new string []
 			{
 				"Core",
 				"Projects",
-				"MetalShaderFormat"
+				"MetalShaderFormat",
+				"ShaderFormatOpenGL",
 			});
 
 		AddThirdPartyPrivateStaticDependencies(Target,
diff --git a/Engine/Source/Programs/CrossCompilerTool/Private/CrossCompilerTool.cpp b/Engine/Source/Programs/CrossCompilerTool/Private/CrossCompilerTool.cpp
--- a/Engine/Source/Programs/CrossCompilerTool/Private/CrossCompilerTool.cpp
+++ b/Engine/Source/Programs/CrossCompilerTool/Private/CrossCompilerTool.cpp
@@ -5,7 +5,7 @@
 #include "CrossCompilerTool.h"
 #include "hlslcc.h"
 #include "MetalBackEnd.h"
-#include "glsl/ir_gen_glsl.h"
+#include "ir_gen_glsl.h"
 
 #include "RequiredProgramMainCPPInclude.h"
 
