UE4 - Mac compile fixes

diff --git a/Engine/Source/Developer/ShaderFormatOpenGL/Private/GlslBackend.cpp b/Engine/Source/Developer/ShaderFormatOpenGL/Private/GlslBackend.cpp
--- a/Engine/Source/Developer/ShaderFormatOpenGL/Private/GlslBackend.cpp
+++ b/Engine/Source/Developer/ShaderFormatOpenGL/Private/GlslBackend.cpp
@@ -47,8 +47,7 @@
 // OpenGL driver chokes on valid GLSL code then.
 
 #if !PLATFORM_WINDOWS
-#define _strnicmp strnicmp
-#define _stricmp stricmp
+#define _strdup strdup
 #endif
 
 static bool GDefaultPrecisionIsHalf = true;
@@ -3271,14 +3270,14 @@ static ir_rvalue* GenShaderInputSemantic(
 	)
 {
 	ir_variable* Variable = NULL;
-	if (Semantic && _strnicmp(Semantic, "SV_", 3) == 0)
+	if (Semantic && FCStringAnsi::Strnicmp(Semantic, "SV_", 3) == 0)
 	{
 		FSystemValue* SystemValues = SystemValueTable[Frequency];
 		for (int i = 0; SystemValues[i].Semantic != NULL; ++i)
 		{
 			if (SystemValues[i].Mode == ir_var_in
 				&& (!SystemValues[i].bESOnly || ParseState->bGenerateES)
-				&& _stricmp(SystemValues[i].Semantic, Semantic) == 0)
+				&& FCStringAnsi::Stricmp(SystemValues[i].Semantic, Semantic) == 0)
 			{
 				if (SystemValues[i].bArrayVariable)
 				{
@@ -3341,7 +3340,7 @@ static ir_rvalue* GenShaderInputSemantic(
 	if (Variable == NULL && Frequency == HSF_DomainShader)
 	{
 		const int PrefixLength = 13;
-		if (_strnicmp(Semantic, "SV_TessFactor", PrefixLength) == 0
+		if (FCStringAnsi::Strnicmp(Semantic, "SV_TessFactor", PrefixLength) == 0
 			&& Semantic[PrefixLength] >= '0'
 			&& Semantic[PrefixLength] <= '3')
 		{
@@ -3357,7 +3356,7 @@ static ir_rvalue* GenShaderInputSemantic(
 	if (Variable == NULL && Frequency == HSF_DomainShader)
 	{
 		const int PrefixLength = 19;
-		if (_strnicmp(Semantic, "SV_InsideTessFactor", PrefixLength) == 0
+		if (FCStringAnsi::Strnicmp(Semantic, "SV_InsideTessFactor", PrefixLength) == 0
 			&& Semantic[PrefixLength] >= '0'
 			&& Semantic[PrefixLength] <= '1')
 		{
@@ -3368,7 +3367,7 @@ static ir_rvalue* GenShaderInputSemantic(
 				ir_var_out
 				);
 		}
-		else if (_stricmp(Semantic, "SV_InsideTessFactor") == 0)
+		else if (FCStringAnsi::Stricmp(Semantic, "SV_InsideTessFactor") == 0)
 		{
 			Variable = new(ParseState)ir_variable(
 				Type,
@@ -3393,7 +3392,7 @@ static ir_rvalue* GenShaderInputSemantic(
 
 	// If we're here, no built-in variables matched.
 
-	if (Semantic && _strnicmp(Semantic, "SV_", 3) == 0)
+	if (Semantic && FCStringAnsi::Strnicmp(Semantic, "SV_", 3) == 0)
 	{
 		_mesa_glsl_warning(ParseState, "unrecognized system "
 			"value input '%s'", Semantic);
@@ -3462,7 +3461,7 @@ static ir_rvalue* GenShaderInputSemantic(
 		if (Frequency == HSF_VertexShader && ParseState->bGenerateLayoutLocations)
 		{
 			const int PrefixLength = 9;
-			if ( (_strnicmp(Semantic, "ATTRIBUTE", PrefixLength) == 0) &&
+			if ( (FCStringAnsi::Strnicmp(Semantic, "ATTRIBUTE", PrefixLength) == 0) &&
 				 (Semantic[PrefixLength] >= '0') &&	 (Semantic[PrefixLength] <= '9') 
 				)
 			{
@@ -3568,14 +3567,14 @@ static ir_rvalue* GenShaderOutputSemantic(
 	FSystemValue* SystemValues = SystemValueTable[Frequency];
 	ir_variable* Variable = NULL;
 
-	if (Semantic && _strnicmp(Semantic, "SV_", 3) == 0)
+	if (Semantic && FCStringAnsi::Strnicmp(Semantic, "SV_", 3) == 0)
 	{
 		for (int i = 0; SystemValues[i].Semantic != NULL; ++i)
 		{
 			if (!SystemValues[i].bESOnly || ParseState->bGenerateES)
 			{
 				if (SystemValues[i].Mode == ir_var_out
-					&& _stricmp(SystemValues[i].Semantic, Semantic) == 0)
+					&& FCStringAnsi::Stricmp(SystemValues[i].Semantic, Semantic) == 0)
 				{
 					Variable = new(ParseState)ir_variable(
 						SystemValues[i].Type,
@@ -3592,7 +3591,7 @@ static ir_rvalue* GenShaderOutputSemantic(
 	if (Variable == NULL && Frequency == HSF_VertexShader)
 	{
 		const int PrefixLength = 15;
-		if (_strnicmp(Semantic, "SV_ClipDistance", PrefixLength) == 0
+		if (FCStringAnsi::Strnicmp(Semantic, "SV_ClipDistance", PrefixLength) == 0
 			&& Semantic[PrefixLength] >= '0'
 			&& Semantic[PrefixLength] <= '9')
 		{
@@ -3608,7 +3607,7 @@ static ir_rvalue* GenShaderOutputSemantic(
 	if (Variable == NULL && Frequency == HSF_PixelShader)
 	{
 		const int PrefixLength = 9;
-		if (_strnicmp(Semantic, "SV_Target", PrefixLength) == 0
+		if (FCStringAnsi::Strnicmp(Semantic, "SV_Target", PrefixLength) == 0
 			&& Semantic[PrefixLength] >= '0'
 			&& Semantic[PrefixLength] <= '7')
 		{
@@ -3630,7 +3629,7 @@ static ir_rvalue* GenShaderOutputSemantic(
 	if (Variable == NULL && Frequency == HSF_HullShader)
 	{
 		const int PrefixLength = 13;
-		if (_strnicmp(Semantic, "SV_TessFactor", PrefixLength) == 0
+		if (FCStringAnsi::Strnicmp(Semantic, "SV_TessFactor", PrefixLength) == 0
 			&& Semantic[PrefixLength] >= '0'
 			&& Semantic[PrefixLength] <= '3')
 		{
@@ -3648,7 +3647,7 @@ static ir_rvalue* GenShaderOutputSemantic(
 	if (Variable == NULL && Frequency == HSF_HullShader)
 	{
 		const int PrefixLength = 19;
-		if (_strnicmp(Semantic, "SV_InsideTessFactor", PrefixLength) == 0
+		if (FCStringAnsi::Strnicmp(Semantic, "SV_InsideTessFactor", PrefixLength) == 0
 			&& Semantic[PrefixLength] >= '0'
 			&& Semantic[PrefixLength] <= '1')
 		{
@@ -3661,7 +3660,7 @@ static ir_rvalue* GenShaderOutputSemantic(
 
 			ApplyClampPowerOfTwo = ParseState->tessellation.partitioning == GLSL_PARTITIONING_POW2;
 		}
-		else if (_stricmp(Semantic, "SV_InsideTessFactor") == 0)
+		else if (FCStringAnsi::Stricmp(Semantic, "SV_InsideTessFactor") == 0)
 		{
 			Variable = new(ParseState)ir_variable(
 				Type,
@@ -3692,7 +3691,7 @@ static ir_rvalue* GenShaderOutputSemantic(
 		return VariableDeref;
 	}
 
-	if (Semantic && _strnicmp(Semantic, "SV_", 3) == 0)
+	if (Semantic && FCStringAnsi::Strnicmp(Semantic, "SV_", 3) == 0)
 	{
 		_mesa_glsl_warning(ParseState, "unrecognized system value output '%s'",
 			Semantic);
@@ -3966,7 +3965,7 @@ static ir_dereference_variable* GenShaderInput(
 	ir_dereference_variable* TempVariableDeref = new(ParseState)ir_dereference_variable(TempVariable);
 	PreCallInstructions->push_tail(TempVariable);
 
-	//check ( InputSemantic ?  (_strnicmp(InputSemantic, "SV_", 3) ==0) : true);
+	//check ( InputSemantic ?  (FCStringAnsi::Strnicmp(InputSemantic, "SV_", 3) ==0) : true);
 
 
 	// everything that's not an Outputpatch is patch constant. System values are treated specially
@@ -4423,7 +4422,7 @@ bool FGlslCodeBackend::GenerateMain(
 		FSystemValue* SystemValues = SystemValueTable[HSF_PixelShader];
 		for (int i = 0; SystemValues[i].Semantic != NULL; ++i)
 		{
-			if (_stricmp(SystemValues[i].GlslName, "gl_FragCoord") == 0)
+			if (FCStringAnsi::Stricmp(SystemValues[i].GlslName, "gl_FragCoord") == 0)
 			{
 				SystemValues[i].bOriginUpperLeft = !ParseState->adjust_clip_space_dx11_to_opengl;
 				break;
diff --git a/Engine/Source/Developer/iOS/MetalShaderFormat/Private/MetalBackend.cpp b/Engine/Source/Developer/iOS/MetalShaderFormat/Private/MetalBackend.cpp
--- a/Engine/Source/Developer/iOS/MetalShaderFormat/Private/MetalBackend.cpp
+++ b/Engine/Source/Developer/iOS/MetalShaderFormat/Private/MetalBackend.cpp
@@ -16,8 +16,7 @@
 #include "MetalUtils.h"
 
 #if !PLATFORM_WINDOWS
-#define _stricmp stricmp
-#define _strnicmp strnicmp
+#define _strdup strdup
 #endif
 
 /**
diff --git a/Engine/Source/Developer/iOS/MetalShaderFormat/Private/MetalUtils.cpp b/Engine/Source/Developer/iOS/MetalShaderFormat/Private/MetalUtils.cpp
--- a/Engine/Source/Developer/iOS/MetalShaderFormat/Private/MetalUtils.cpp
+++ b/Engine/Source/Developer/iOS/MetalShaderFormat/Private/MetalUtils.cpp
@@ -16,11 +16,6 @@
 #include "ir_optimization.h"
 #include "MetalUtils.h"
 
-#if !PLATFORM_WINDOWS
-#define _stricmp stricmp
-#define _strnicmp strnicmp
-#endif
-
 const bool bExpandVSInputsToFloat4 = false;
 const bool bGenerateVSInputDummies = false;
 
@@ -281,12 +276,12 @@ namespace MetalUtils
 		}
 
 		ir_variable* Variable = NULL;
-		if (_strnicmp(Semantic, "SV_", 3) == 0)
+		if (FCStringAnsi::Strnicmp(Semantic, "SV_", 3) == 0)
 		{
 			FSystemValue* SystemValues = SystemValueTable[Frequency];
 			for (int i = 0; SystemValues[i].HlslSemantic != NULL; ++i)
 			{
-				if (SystemValues[i].Mode == ir_var_in && _stricmp(SystemValues[i].HlslSemantic, Semantic) == 0)
+				if (SystemValues[i].Mode == ir_var_in && FCStringAnsi::Stricmp(SystemValues[i].HlslSemantic, Semantic) == 0)
 				{
 					ir_variable* Variable = new(ParseState)ir_variable(
 						SystemValues[i].Type, SystemValues[i].MetalName, ir_var_in);
@@ -303,7 +298,7 @@ namespace MetalUtils
 		}
 
 		// If we're here, no built-in variables matched.
-		if (_strnicmp(Semantic, "SV_", 3) == 0)
+		if (FCStringAnsi::Strnicmp(Semantic, "SV_", 3) == 0)
 		{
 			_mesa_glsl_warning(ParseState, "unrecognized system value input '%s'", Semantic);
 		}
@@ -314,7 +309,7 @@ namespace MetalUtils
 			ir_var_in);
 		if (Frequency == HSF_VertexShader)
 		{
-			if (!_strnicmp(Semantic, "ATTRIBUTE", 9))
+			if (!FCStringAnsi::Strnicmp(Semantic, "ATTRIBUTE", 9))
 			{
 				Variable->semantic = ralloc_asprintf(ParseState, "[[ attribute(%s) ]]", Semantic);
 			}
@@ -437,12 +432,12 @@ namespace MetalUtils
 	{
 		ir_variable* Variable = NULL;
 
-		if (!Variable && _strnicmp(Semantic, "SV_", 3) == 0)
+		if (!Variable && FCStringAnsi::Strnicmp(Semantic, "SV_", 3) == 0)
 		{
 			FSystemValue* SystemValues = SystemValueTable[Frequency];
 			for (int i = 0; SystemValues[i].HlslSemantic != nullptr; ++i)
 			{
-				if (SystemValues[i].Mode == ir_var_out && _stricmp(SystemValues[i].HlslSemantic, Semantic) == 0)
+				if (SystemValues[i].Mode == ir_var_out && FCStringAnsi::Stricmp(SystemValues[i].HlslSemantic, Semantic) == 0)
 				{
 					Variable = new(ParseState) ir_variable(SystemValues[i].Type, SystemValues[i].MetalName, ir_var_out);
 					Variable->semantic = SystemValues[i].MetalSemantic;
@@ -451,7 +446,7 @@ namespace MetalUtils
 			}
 		}
 
-		if (Semantic && _strnicmp(Semantic, "SV_", 3) == 0 && !Variable)
+		if (Semantic && FCStringAnsi::Strnicmp(Semantic, "SV_", 3) == 0 && !Variable)
 		{
 			_mesa_glsl_warning(ParseState, "unrecognized system value output '%s'", Semantic);
 		}
@@ -916,7 +911,7 @@ static bool ProcessStageInVariables(_mesa_glsl_parse_state* ParseState, EHlslSha
 	const auto* SystemValues = MetalUtils::SystemValueTable[Frequency];
 	for(int i = 0; SystemValues[i].MetalSemantic != nullptr; ++i)
 	{
-		if (SystemValues[i].Mode == ir_var_in && _stricmp(SystemValues[i].MetalSemantic, Variable->semantic) == 0)
+		if (SystemValues[i].Mode == ir_var_in && FCStringAnsi::Stricmp(SystemValues[i].MetalSemantic, Variable->semantic) == 0)
 		{
 			return true;
 		}
@@ -1152,14 +1147,14 @@ static ir_rvalue* GenShaderInputSemantic(
 	int SemanticArrayIndex)
 {
 	ir_variable* Variable = NULL;
-	if (Semantic && _strnicmp(Semantic, "SV_", 3) == 0)
+	if (Semantic && FCStringAnsi::Strnicmp(Semantic, "SV_", 3) == 0)
 	{
 		FSystemValue* SystemValues = SystemValueTable[Frequency];
 		for (int i = 0; SystemValues[i].Semantic != NULL; ++i)
 		{
 			if (SystemValues[i].Mode == ir_var_in
 				//				&& (!SystemValues[i].bESOnly || ParseState->bGenerateES)
-				&& _stricmp(SystemValues[i].Semantic, Semantic) == 0)
+				&& FCStringAnsi::Stricmp(SystemValues[i].Semantic, Semantic) == 0)
 			{
 				check(0);
 #if 0
@@ -1228,7 +1223,7 @@ static ir_rvalue* GenShaderInputSemantic(
 
 	// If we're here, no built-in variables matched.
 
-	if (Semantic && _strnicmp(Semantic, "SV_", 3) == 0)
+	if (Semantic && FCStringAnsi::Strnicmp(Semantic, "SV_", 3) == 0)
 	{
 		_mesa_glsl_warning(ParseState, "unrecognized system "
 			"value input '%s'", Semantic);
@@ -1412,12 +1407,12 @@ static ir_rvalue* GenShaderOutputSemantic(
 	FSystemValue* SystemValues = SystemValueTable[Frequency];
 	ir_variable* Variable = NULL;
 
-	if (Semantic && _strnicmp(Semantic, "SV_", 3) == 0)
+	if (Semantic && FCStringAnsi::Strnicmp(Semantic, "SV_", 3) == 0)
 	{
 		for (int i = 0; SystemValues[i].Semantic != NULL; ++i)
 		{
 			if (SystemValues[i].Mode == ir_var_out
-				&& _stricmp(SystemValues[i].Semantic, Semantic) == 0)
+				&& FCStringAnsi::Stricmp(SystemValues[i].Semantic, Semantic) == 0)
 			{
 				check(0);
 			}
@@ -1427,7 +1422,7 @@ static ir_rvalue* GenShaderOutputSemantic(
 	if (Variable == NULL && Frequency == HSF_VertexShader)
 	{
 		const int PrefixLength = 15;
-		if (_strnicmp(Semantic, "SV_ClipDistance", PrefixLength) == 0
+		if (FCStringAnsi::Strnicmp(Semantic, "SV_ClipDistance", PrefixLength) == 0
 			&& Semantic[PrefixLength] >= '0'
 			&& Semantic[PrefixLength] <= '9')
 		{
@@ -1438,7 +1433,7 @@ static ir_rvalue* GenShaderOutputSemantic(
 	if (Variable == NULL && Frequency == HSF_PixelShader)
 	{
 		const int PrefixLength = 9;
-		if (_strnicmp(Semantic, "SV_Target", PrefixLength) == 0
+		if (FCStringAnsi::Strnicmp(Semantic, "SV_Target", PrefixLength) == 0
 			&& Semantic[PrefixLength] >= '0'
 			&& Semantic[PrefixLength] <= '7')
 		{
@@ -1471,7 +1466,7 @@ static ir_rvalue* GenShaderOutputSemantic(
 		return VariableDeref;
 	}
 
-	if (Semantic && _strnicmp(Semantic, "SV_", 3) == 0)
+	if (Semantic && FCStringAnsi::Strnicmp(Semantic, "SV_", 3) == 0)
 	{
 		_mesa_glsl_warning(ParseState, "unrecognized system value output '%s'",
 			Semantic);
