Fix-up the .w component for gl_FragCoord being used as SV_Position - it would appear that they aren't the same (GL claims 1/clip.w, presumably from the usage in our shaders D3D provides clip.w).

diff --git a/Engine/Source/Developer/ShaderFormatOpenGL/Private/GlslBackend.cpp b/Engine/Source/Developer/ShaderFormatOpenGL/Private/GlslBackend.cpp
--- a/Engine/Source/Developer/ShaderFormatOpenGL/Private/GlslBackend.cpp
+++ b/Engine/Source/Developer/ShaderFormatOpenGL/Private/GlslBackend.cpp
@@ -3319,7 +3319,37 @@ static ir_rvalue* GenShaderInputSemantic(
 						ParseState->symbols->add_variable(Variable);
 						ir_dereference_variable* VariableDeref = new(ParseState)ir_dereference_variable(Variable);
 
-						if (ParseState->adjust_clip_space_dx11_to_opengl && SystemValues[i].bApplyClipSpaceAdjustment)
+						if (FCStringAnsi::Stricmp(Semantic, "SV_Position") == 0 && Frequency == HSF_PixelShader)
+						{
+							// This is for input of gl_FragCoord into pixel shader only.
+							
+							// Generate a local variable to do the conversion in, keeping source type.
+							ir_variable* TempVariable = new(ParseState)ir_variable(Variable->type, NULL, ir_var_temporary);
+							DeclInstructions->push_tail(TempVariable);
+							
+							// Assign input to this variable
+							ir_dereference_variable* TempVariableDeref = new(ParseState)ir_dereference_variable(TempVariable);
+							DeclInstructions->push_tail(
+								new(ParseState)ir_assignment(
+								TempVariableDeref,
+								VariableDeref
+								)
+							);
+							
+							// TempVariable.w = ( 1.0f / TempVariable.w );
+							DeclInstructions->push_tail(
+								new(ParseState)ir_assignment(
+								new(ParseState)ir_swizzle(TempVariableDeref->clone(ParseState, NULL), 3, 0, 0, 0, 1),
+								new(ParseState)ir_expression(ir_binop_div,
+								new(ParseState)ir_constant(1.0f),
+								new(ParseState)ir_swizzle(TempVariableDeref->clone(ParseState, NULL), 3, 0, 0, 0, 1)
+								)
+								)
+								);
+
+							return TempVariableDeref->clone(ParseState, NULL);
+						}
+						else if (ParseState->adjust_clip_space_dx11_to_opengl && SystemValues[i].bApplyClipSpaceAdjustment)
 						{
 							// incoming gl_FrontFacing. Make it (!gl_FrontFacing), due to vertical flip in OpenGL
 							return new(ParseState)ir_expression(ir_unop_logic_not, glsl_type::bool_type, VariableDeref, NULL);
