Update the version numbers for OpenGL shader formats as CL #2460109 has altered the cached data structures so they need to be recompiled.
Once SCW is recompiled this will prevent any further crashes due to incompatible DDC OpenGL shader binaries.

diff --git a/Engine/Source/Developer/ShaderFormatOpenGL/Private/ShaderFormatOpenGL.cpp b/Engine/Source/Developer/ShaderFormatOpenGL/Private/ShaderFormatOpenGL.cpp
--- a/Engine/Source/Developer/ShaderFormatOpenGL/Private/ShaderFormatOpenGL.cpp
+++ b/Engine/Source/Developer/ShaderFormatOpenGL/Private/ShaderFormatOpenGL.cpp
@@ -22,14 +22,14 @@ class FShaderFormatGLSL : public IShaderFormat
 	enum 
 	{
 		/** Version for shader format, this becomes part of the DDC key. */
-		UE_SHADER_GLSL_150_VER = 25,
-		UE_SHADER_GLSL_150_MAC_VER = 5,
-		UE_SHADER_GLSL_430_VER = 55,
-		UE_SHADER_GLSL_ES2_VER = 13,
-		UE_SHADER_GLSL_150ES2_VER = 16,
-		UE_SHADER_GLSL_ES2_VER_WEBGL  = 14,
-		UE_SHADER_GLSL_ES2_IOS_VER  = 4,
-		UE_SHADER_GLSL_310_ES_EXT_VER = 2,
+		UE_SHADER_GLSL_150_VER = 26,
+		UE_SHADER_GLSL_150_MAC_VER = 6,
+		UE_SHADER_GLSL_430_VER = 56,
+		UE_SHADER_GLSL_ES2_VER = 14,
+		UE_SHADER_GLSL_150ES2_VER = 17,
+		UE_SHADER_GLSL_ES2_VER_WEBGL  = 15,
+		UE_SHADER_GLSL_ES2_IOS_VER  = 5,
+		UE_SHADER_GLSL_310_ES_EXT_VER = 3,
 	}; 
 
 	void CheckFormat(FName Format) const
