UE4 - OpenGL - Allow running emulated (packed/flattened) uniform buffers when running PC ES2; use OpenGL.UseEmulatedUBs=1 on your consolevariables.ini
- Only works with -game, not in-editor!

diff --git a/Engine/Source/Developer/ShaderFormatOpenGL/Private/OpenGLShaderCompiler.cpp b/Engine/Source/Developer/ShaderFormatOpenGL/Private/OpenGLShaderCompiler.cpp
--- a/Engine/Source/Developer/ShaderFormatOpenGL/Private/OpenGLShaderCompiler.cpp
+++ b/Engine/Source/Developer/ShaderFormatOpenGL/Private/OpenGLShaderCompiler.cpp
@@ -64,12 +64,12 @@ DEFINE_LOG_CATEGORY_STATIC(LogOpenGLShaderCompiler, Log, All);
   
 static FORCEINLINE bool IsES2Platform(GLSLVersion Version)
 {
-	return (Version == GLSL_ES2 || Version == GLSL_150_ES2 || Version == GLSL_ES2_WEBGL || Version == GLSL_ES2_IOS); 
+	return (Version == GLSL_ES2 || Version == GLSL_150_ES2 || Version == GLSL_ES2_WEBGL || Version == GLSL_ES2_IOS || Version == GLSL_150_ES2_NOUB); 
 }
 
 static FORCEINLINE bool IsPCES2Platform(GLSLVersion Version)
 {
-	return (Version == GLSL_150_ES2);
+	return (Version == GLSL_150_ES2 || Version == GLSL_150_ES2_NOUB);
 }
 
 // This function should match OpenGLShaderPlatformSeparable
@@ -78,7 +78,7 @@ static FORCEINLINE bool SupportsSeparateShaderObjects(GLSLVersion Version)
 	// Only desktop shader platforms can use separable shaders for now,
 	// the generated code relies on macros supplied at runtime to determine whether
 	// shaders may be separable and/or linked.
-	return Version == GLSL_150 || Version == GLSL_150_MAC || Version == GLSL_150_ES2 || Version == GLSL_430;
+	return Version == GLSL_150 || Version == GLSL_150_MAC || Version == GLSL_150_ES2 || Version == GLSL_150_ES2_NOUB|| Version == GLSL_430;
 }
 
 /*------------------------------------------------------------------------------
@@ -1188,6 +1188,7 @@ static void OpenGLVersionFromGLSLVersion(GLSLVersion InVersion, int& OutMajorVer
 			OutMinorVersion = 3;
 			break;
 		case GLSL_150_ES2:
+		case GLSL_150_ES2_NOUB:
 			OutMajorVersion = 3;
 			OutMinorVersion = 2;
 			break;
@@ -1516,6 +1517,7 @@ static FString CreateCrossCompilerBatchFile( const FString& ShaderFile, const FS
 	switch (Version)
 	{
 		case GLSL_150:
+		case GLSL_150_ES2:
 			VersionSwitch = TEXT(" -gl3 -separateshaders");
 			break;
 
@@ -1523,7 +1525,7 @@ static FString CreateCrossCompilerBatchFile( const FString& ShaderFile, const FS
 			VersionSwitch = TEXT(" -gl3 -mac -separateshaders");
 			break;
 
-		case GLSL_150_ES2:
+		case GLSL_150_ES2_NOUB:
 			VersionSwitch = TEXT(" -gl3 -flattenub -flattenubstruct -separateshaders");
 			break;
 
@@ -1614,6 +1616,7 @@ void CompileShader_Windows_OGL(const FShaderCompilerInput& Input,FShaderCompiler
 			break; 
 
 		case GLSL_150_ES2:
+		case GLSL_150_ES2_NOUB:
 			AdditionalDefines.SetDefine(TEXT("COMPILER_GLSL"), 1);
 			AdditionalDefines.SetDefine(TEXT("ES2_PROFILE"), 1);
 			HlslCompilerTarget = HCT_FeatureLevelSM4;
@@ -1683,7 +1686,12 @@ void CompileShader_Windows_OGL(const FShaderCompilerInput& Input,FShaderCompiler
 			CCFlags |= HLSLCC_ApplyCommonSubexpressionElimination;
 		}
 		
-		if(SupportsSeparateShaderObjects(Version))
+		if (Version == GLSL_150_ES2_NOUB)
+		{
+			CCFlags |= HLSLCC_FlattenUniformBuffers | HLSLCC_FlattenUniformBufferStructures;
+		}
+		
+		if (SupportsSeparateShaderObjects(Version))
 		{
 			CCFlags |= HLSLCC_SeparateShaderObjects;
 		}
@@ -1750,13 +1758,10 @@ void CompileShader_Windows_OGL(const FShaderCompilerInput& Input,FShaderCompiler
 
 #if VALIDATE_GLSL_WITH_DRIVER
 			PrecompileShader(Output, Input, GlslShaderSource, Version, Frequency);
-			if (Output.bSucceeded == false)
-			{
-			}
 #else // VALIDATE_GLSL_WITH_DRIVER
 			int32 SourceLen = FCStringAnsi::Strlen(GlslShaderSource);
 			Output.Target = Input.Target;
-				BuildShaderOutput(Output, Input, GlslShaderSource, SourceLen, Version);
+			BuildShaderOutput(Output, Input, GlslShaderSource, SourceLen, Version);
 #endif // VALIDATE_GLSL_WITH_DRIVER
 		}
 		else
diff --git a/Engine/Source/Developer/ShaderFormatOpenGL/Private/ShaderFormatOpenGL.cpp b/Engine/Source/Developer/ShaderFormatOpenGL/Private/ShaderFormatOpenGL.cpp
--- a/Engine/Source/Developer/ShaderFormatOpenGL/Private/ShaderFormatOpenGL.cpp
+++ b/Engine/Source/Developer/ShaderFormatOpenGL/Private/ShaderFormatOpenGL.cpp
@@ -14,6 +14,7 @@ static FName NAME_GLSL_430(TEXT("GLSL_430"));
 static FName NAME_GLSL_ES2(TEXT("GLSL_ES2"));
 static FName NAME_GLSL_ES2_WEBGL(TEXT("GLSL_ES2_WEBGL"));
 static FName NAME_GLSL_150_ES2(TEXT("GLSL_150_ES2"));
+static FName NAME_GLSL_150_ES2_NOUB(TEXT("GLSL_150_ES2_NOUB"));
 static FName NAME_GLSL_ES2_IOS(TEXT("GLSL_ES2_IOS"));
 static FName NAME_GLSL_310_ES_EXT(TEXT("GLSL_310_ES_EXT"));
  
@@ -27,6 +28,7 @@ class FShaderFormatGLSL : public IShaderFormat
 		UE_SHADER_GLSL_430_VER = 56,
 		UE_SHADER_GLSL_ES2_VER = 14,
 		UE_SHADER_GLSL_150ES2_VER = 17,
+		UE_SHADER_GLSL_150ES2NOUB_VER = 1,
 		UE_SHADER_GLSL_ES2_VER_WEBGL  = 15,
 		UE_SHADER_GLSL_ES2_IOS_VER  = 5,
 		UE_SHADER_GLSL_310_ES_EXT_VER = 3,
@@ -39,6 +41,7 @@ class FShaderFormatGLSL : public IShaderFormat
 				Format == NAME_GLSL_430 || 
 				Format == NAME_GLSL_ES2 || 
 				Format == NAME_GLSL_150_ES2 ||
+				Format == NAME_GLSL_150_ES2_NOUB ||
                 Format == NAME_GLSL_ES2_WEBGL ||
 				Format == NAME_GLSL_ES2_IOS ||
 				Format == NAME_GLSL_310_ES_EXT
@@ -70,6 +73,10 @@ public:
 		{
 			GLSLVersion = UE_SHADER_GLSL_150ES2_VER;
 		}
+		else if (Format == NAME_GLSL_150_ES2_NOUB)
+		{
+			GLSLVersion = UE_SHADER_GLSL_150ES2NOUB_VER;
+		}
 		else if (Format == NAME_GLSL_ES2_WEBGL)
 		{
 			GLSLVersion = UE_SHADER_GLSL_ES2_VER_WEBGL;
@@ -100,6 +107,7 @@ public:
 		OutFormats.Add(NAME_GLSL_150_ES2);
 		OutFormats.Add(NAME_GLSL_ES2_IOS);
 		OutFormats.Add(NAME_GLSL_310_ES_EXT);
+		OutFormats.Add(NAME_GLSL_150_ES2_NOUB);
 	}
 	virtual void CompileShader(FName Format, const struct FShaderCompilerInput& Input, struct FShaderCompilerOutput& Output,const FString& WorkingDirectory) const
 	{
@@ -143,6 +151,10 @@ public:
 		{
 			CompileShader_Windows_OGL(Input, Output, WorkingDirectory, GLSL_ES2_IOS);
 		}
+		else if (Format == NAME_GLSL_150_ES2_NOUB)
+		{
+			CompileShader_Windows_OGL(Input, Output, WorkingDirectory, GLSL_150_ES2_NOUB);
+		}
 		else if (Format == NAME_GLSL_150_ES2)
 		{
 			CompileShader_Windows_OGL(Input, Output, WorkingDirectory, GLSL_150_ES2);
diff --git a/Engine/Source/Developer/ShaderFormatOpenGL/Private/ShaderFormatOpenGL.h b/Engine/Source/Developer/ShaderFormatOpenGL/Private/ShaderFormatOpenGL.h
--- a/Engine/Source/Developer/ShaderFormatOpenGL/Private/ShaderFormatOpenGL.h
+++ b/Engine/Source/Developer/ShaderFormatOpenGL/Private/ShaderFormatOpenGL.h
@@ -9,6 +9,7 @@ enum GLSLVersion
 	GLSL_ES2,
 	GLSL_ES2_WEBGL,
 	GLSL_150_ES2,	// ES2 Emulation
+	GLSL_150_ES2_NOUB,	// ES2 Emulation with NoUBs
 	GLSL_ES2_IOS,
 	GLSL_150_MAC, // Apple only
 	GLSL_310_ES_EXT,
diff --git a/Engine/Source/Runtime/OpenGLDrv/Private/OpenGLCommands.cpp b/Engine/Source/Runtime/OpenGLDrv/Private/OpenGLCommands.cpp
--- a/Engine/Source/Runtime/OpenGLDrv/Private/OpenGLCommands.cpp
+++ b/Engine/Source/Runtime/OpenGLDrv/Private/OpenGLCommands.cpp
@@ -39,6 +39,12 @@ namespace OpenGLConsoleVariables
 		TEXT("If true, use glMapBuffer otherwise use glBufferSubdata.")
 		);
 
+	static FAutoConsoleVariable CVarUseEmulatedUBs(
+		TEXT("OpenGL.UseEmulatedUBs"),
+		0,
+		TEXT("If true, enable using emulated uniform buffers on ES2 mode."),
+		ECVF_ReadOnly
+		);
 
 	int32 bSkipCompute = 0;
 	static FAutoConsoleVariableRef CVarSkipCompute(
diff --git a/Engine/Source/Runtime/OpenGLDrv/Private/OpenGLDevice.cpp b/Engine/Source/Runtime/OpenGLDrv/Private/OpenGLDevice.cpp
--- a/Engine/Source/Runtime/OpenGLDrv/Private/OpenGLDevice.cpp
+++ b/Engine/Source/Runtime/OpenGLDrv/Private/OpenGLDevice.cpp
@@ -659,6 +659,11 @@ static void InitRHICapabilitiesForGL()
 
 	// Emulate uniform buffers on ES2, unless we're on a desktop platform emulating ES2.
 	GUseEmulatedUniformBuffers = IsES2Platform(GMaxRHIShaderPlatform) && !IsPCPlatform(GMaxRHIShaderPlatform);
+	if (!GUseEmulatedUniformBuffers && IsPCPlatform(GMaxRHIShaderPlatform))
+	{
+		static auto* CVar = IConsoleManager::Get().FindTConsoleVariableDataInt(TEXT("OpenGL.UseEmulatedUBs"));
+		GUseEmulatedUniformBuffers = CVar && CVar->GetValueOnAnyThread() != 0;
+	}
 
 	FString FeatureLevelName;
 	GetFeatureLevelName(GMaxRHIFeatureLevel, FeatureLevelName);
diff --git a/Engine/Source/Runtime/RHI/Private/RHI.cpp b/Engine/Source/Runtime/RHI/Private/RHI.cpp
--- a/Engine/Source/Runtime/RHI/Private/RHI.cpp
+++ b/Engine/Source/Runtime/RHI/Private/RHI.cpp
@@ -271,6 +271,7 @@ static FName NAME_SF_PS4(TEXT("SF_PS4"));
 static FName NAME_SF_XBOXONE(TEXT("SF_XBOXONE"));
 static FName NAME_GLSL_430(TEXT("GLSL_430"));
 static FName NAME_OPENGL_150_ES2(TEXT("GLSL_150_ES2"));
+static FName NAME_OPENGL_150_ES2_NOUB(TEXT("GLSL_150_ES2_NOUB"));
 static FName NAME_OPENGL_ES2(TEXT("GLSL_ES2"));
 static FName NAME_OPENGL_ES2_WEBGL(TEXT("GLSL_ES2_WEBGL"));
 static FName NAME_OPENGL_ES2_IOS(TEXT("GLSL_ES2_IOS"));
@@ -299,7 +300,10 @@ FName LegacyShaderPlatformToShaderFormat(EShaderPlatform Platform)
 	case SP_OPENGL_SM5:
 		return NAME_GLSL_430;
 	case SP_OPENGL_PCES2:
-		return NAME_OPENGL_150_ES2;
+	{
+		static auto* CVar = IConsoleManager::Get().FindTConsoleVariableDataInt(TEXT("OpenGL.UseEmulatedUBs"));
+		return (CVar && CVar->GetValueOnAnyThread() != 0) ? NAME_OPENGL_150_ES2_NOUB : NAME_OPENGL_150_ES2;
+	}
 	case SP_OPENGL_ES2:
 		return NAME_OPENGL_ES2;
 	case SP_OPENGL_ES2_WEBGL:
@@ -326,13 +330,14 @@ EShaderPlatform ShaderFormatToLegacyShaderPlatform(FName ShaderFormat)
 	if (ShaderFormat == NAME_PCD3D_ES2)			return SP_PCD3D_ES2;
 	if (ShaderFormat == NAME_GLSL_150)			return SP_OPENGL_SM4;
 	if (ShaderFormat == NAME_GLSL_150_MAC)		return SP_OPENGL_SM4_MAC;
-	if (ShaderFormat == NAME_SF_PS4)			return SP_PS4;
-	if (ShaderFormat == NAME_SF_XBOXONE)		return SP_XBOXONE;
+	if (ShaderFormat == NAME_SF_PS4)				return SP_PS4;
+	if (ShaderFormat == NAME_SF_XBOXONE)			return SP_XBOXONE;
 	if (ShaderFormat == NAME_GLSL_430)			return SP_OPENGL_SM5;
-	if (ShaderFormat == NAME_OPENGL_150_ES2)	return SP_OPENGL_PCES2;
-	if (ShaderFormat == NAME_OPENGL_ES2)		return SP_OPENGL_ES2;
+	if (ShaderFormat == NAME_OPENGL_150_ES2 || ShaderFormat == NAME_OPENGL_150_ES2_NOUB)
+												return SP_OPENGL_PCES2;
+	if (ShaderFormat == NAME_OPENGL_ES2)			return SP_OPENGL_ES2;
 	if (ShaderFormat == NAME_OPENGL_ES2_WEBGL)	return SP_OPENGL_ES2_WEBGL;
-	if (ShaderFormat == NAME_OPENGL_ES2_IOS)	return SP_OPENGL_ES2_IOS;
+	if (ShaderFormat == NAME_OPENGL_ES2_IOS)		return SP_OPENGL_ES2_IOS;
 	if (ShaderFormat == NAME_SF_METAL)			return SP_METAL;
 	if (ShaderFormat == NAME_SF_METAL_MRT)		return SP_METAL_MRT;
 	if (ShaderFormat == NAME_GLSL_310_ES_EXT)	return SP_OPENGL_ES31_EXT;
