UE4 - OpenGL - Generate GL_EXT_tessellation_shader properly (UE-11222)

diff --git a/Engine/Source/Developer/ShaderFormatOpenGL/Private/GlslBackend.cpp b/Engine/Source/Developer/ShaderFormatOpenGL/Private/GlslBackend.cpp
--- a/Engine/Source/Developer/ShaderFormatOpenGL/Private/GlslBackend.cpp
+++ b/Engine/Source/Developer/ShaderFormatOpenGL/Private/GlslBackend.cpp
@@ -2728,7 +2728,7 @@ class ir_gen_glsl_visitor : public ir_visitor
 #endif		
 	}
 
-	void print_extensions(bool bUsesFramebufferFetchES2, bool bUsesES31Extensions)
+	void print_extensions(_mesa_glsl_parse_state* state, bool bUsesFramebufferFetchES2, bool bUsesES31Extensions)
 	{
 		if (bUsesES2TextureLODExtension)
 		{
@@ -2737,6 +2737,12 @@ class ir_gen_glsl_visitor : public ir_visitor
 			ralloc_asprintf_append(buffer, "#endif\n");
 		}
 
+		if (state->bSeparateShaderObjects && !state->bGenerateES && 
+			(state->target == tessellation_control_shader) || (state->target == tessellation_evaluation_shader))
+		{
+			ralloc_asprintf_append(buffer, "#extension GL_ARB_tessellation_shader : enable\n");
+		}
+
 		if (bUsesDXDY && bIsES)
 		{
 			ralloc_asprintf_append(buffer, "#extension GL_OES_standard_derivatives : enable\n");
@@ -2878,8 +2884,8 @@ public:
 
 		char* Extensions = ralloc_asprintf(mem_ctx, "");
 		buffer = &Extensions;
-		print_extensions(bUsesFramebufferFetchES2, state->language_version == 310);
-		if(state->bSeparateShaderObjects && !state->bGenerateES)
+		print_extensions(state, bUsesFramebufferFetchES2, state->language_version == 310);
+		if (state->bSeparateShaderObjects && !state->bGenerateES)
 		{
 			switch (state->target)
 			{
diff --git a/Engine/Source/Developer/ShaderFormatOpenGL/Private/ShaderFormatOpenGL.cpp b/Engine/Source/Developer/ShaderFormatOpenGL/Private/ShaderFormatOpenGL.cpp
--- a/Engine/Source/Developer/ShaderFormatOpenGL/Private/ShaderFormatOpenGL.cpp
+++ b/Engine/Source/Developer/ShaderFormatOpenGL/Private/ShaderFormatOpenGL.cpp
@@ -25,7 +25,7 @@ class FShaderFormatGLSL : public IShaderFormat
 		/** Version for shader format, this becomes part of the DDC key. */
 		UE_SHADER_GLSL_150_VER = 26,
 		UE_SHADER_GLSL_150_MAC_VER = 6,
-		UE_SHADER_GLSL_430_VER = 56,
+		UE_SHADER_GLSL_430_VER = 57,
 		UE_SHADER_GLSL_ES2_VER = 14,
 		UE_SHADER_GLSL_150ES2_VER = 17,
 		UE_SHADER_GLSL_150ES2NOUB_VER = 1,
diff --git a/Engine/Source/Runtime/OpenGLDrv/Private/OpenGLShaders.cpp b/Engine/Source/Runtime/OpenGLDrv/Private/OpenGLShaders.cpp
--- a/Engine/Source/Runtime/OpenGLDrv/Private/OpenGLShaders.cpp
+++ b/Engine/Source/Runtime/OpenGLDrv/Private/OpenGLShaders.cpp
@@ -563,14 +563,6 @@ ShaderType* CompileOpenGLShader(const TArray<uint8>& Code)
 			}
 		}
 
-		// Add an extension directive (in line with https://www.opengl.org/registry/specs/ARB/tessellation_shader.txt)
-		// Note that this is a temporary hack, proper resolution is tracked as UE-11222
-		bool bNeedsTesselationExtension = (FCStringAnsi::Strstr(GlslCodeOriginal.GetData(), "gl_MaxPatchVertices") != nullptr);
-		if (bNeedsTesselationExtension)
-		{
-			AppendCString(GlslCode, "#extension GL_ARB_tessellation_shader : enable\n");
-		}
-
 #if PLATFORM_ANDROID 
 		// Temporary patch to remove #extension GL_OES_standard_derivaties if not supported
 		if (!FOpenGL::SupportsStandardDerivativesExtension())
