UE4 - Metal - Fix bad uav dereference
Glsl - Better variable names

diff --git a/Engine/Source/Developer/ShaderFormatOpenGL/Private/GlslBackend.cpp b/Engine/Source/Developer/ShaderFormatOpenGL/Private/GlslBackend.cpp
--- a/Engine/Source/Developer/ShaderFormatOpenGL/Private/GlslBackend.cpp
+++ b/Engine/Source/Developer/ShaderFormatOpenGL/Private/GlslBackend.cpp
@@ -563,9 +563,32 @@ class ir_gen_glsl_visitor : public ir_visitor
 			const char *name = (const char *) hash_table_find(this->printable_names, var);
 			if (name == NULL)
 			{
-				bool is_global = (scope_depth == 0 && var->mode != ir_var_temporary);
-				const char *prefix = is_global ? "g" : "t";
-				int var_id = is_global ? global_id++ : temp_id++;
+				bool bIsGlobal = (scope_depth == 0 && var->mode != ir_var_temporary);
+				const char* prefix = "g";
+				if (!bIsGlobal)
+				{
+					if (var->type->is_matrix())
+					{
+						prefix = "m";
+					}
+					else if (var->type->is_vector())
+					{
+						prefix = "v";
+					}
+					else
+					{
+						switch (var->type->base_type)
+						{
+						case GLSL_TYPE_BOOL: prefix = "b"; break;
+						case GLSL_TYPE_UINT: prefix = "u"; break;
+						case GLSL_TYPE_INT: prefix = "i"; break;
+						case GLSL_TYPE_HALF: prefix = "h"; break;
+						case GLSL_TYPE_FLOAT: prefix = "f"; break;
+						default: prefix = "t"; break;
+						}
+					}
+				}
+				int var_id = bIsGlobal ? global_id++ : temp_id++;
 				name = ralloc_asprintf(mem_ctx, "%s%d", prefix, var_id);
 				hash_table_insert(this->printable_names, (void *)name, var);
 			}
diff --git a/Engine/Source/Developer/iOS/MetalShaderFormat/Private/MetalBackend.cpp b/Engine/Source/Developer/iOS/MetalShaderFormat/Private/MetalBackend.cpp
--- a/Engine/Source/Developer/iOS/MetalShaderFormat/Private/MetalBackend.cpp
+++ b/Engine/Source/Developer/iOS/MetalShaderFormat/Private/MetalBackend.cpp
@@ -1315,7 +1315,6 @@ protected:
 				}
 				else
 				{
-					deref->image->accept(this);
 					ralloc_asprintf_append( buffer, "[" );
 					deref->image_index->accept(this);
 					ralloc_asprintf_append( buffer, "] = " );
