UE4 - Remove opengl precision modifiers from non-es platforms

diff --git a/Engine/Source/Developer/ShaderFormatOpenGL/Private/GlslBackend.cpp b/Engine/Source/Developer/ShaderFormatOpenGL/Private/GlslBackend.cpp
--- a/Engine/Source/Developer/ShaderFormatOpenGL/Private/GlslBackend.cpp
+++ b/Engine/Source/Developer/ShaderFormatOpenGL/Private/GlslBackend.cpp
@@ -936,22 +936,25 @@ class ir_gen_glsl_visitor : public ir_visitor
 						);
 				}
 
-				if (GDefaultPrecisionIsHalf && var->type->inner_type->base_type == GLSL_TYPE_FLOAT)
+				if (bIsES)
 				{
-					ralloc_asprintf_append(buffer, "highp ");
-				}
-				else if (!GDefaultPrecisionIsHalf && var->type->inner_type->base_type == GLSL_TYPE_HALF)
-				{
-					ralloc_asprintf_append(buffer, "mediump ");
-				}
-				else if (var->type->inner_type->is_integer())
-				{
-					ralloc_asprintf_append(buffer, "highp ");
-				}
-				else
-				{
-					// we missed a type
-					check(false);
+					if (GDefaultPrecisionIsHalf && var->type->inner_type->base_type == GLSL_TYPE_FLOAT)
+					{
+						ralloc_asprintf_append(buffer, "highp ");
+					}
+					else if (!GDefaultPrecisionIsHalf && var->type->inner_type->base_type == GLSL_TYPE_HALF)
+					{
+						ralloc_asprintf_append(buffer, "mediump ");
+					}
+					else if (var->type->inner_type->is_integer())
+					{
+						ralloc_asprintf_append(buffer, "highp ");
+					}
+					else
+					{
+						// we missed a type
+						check(false);
+					}
 				}
 				print_type_pre(var->type);
 			}
@@ -976,33 +979,36 @@ class ir_gen_glsl_visitor : public ir_visitor
 					mode_str[var->mode],
 					interp_str[var->interpolation]
 					);
-				if (var->type->is_sampler())
+				if (bIsES)
 				{
-					if (GDefaultPrecisionIsHalf && var->type->inner_type->base_type == GLSL_TYPE_FLOAT)
+					if (var->type->is_sampler())
+					{
+						if (GDefaultPrecisionIsHalf && var->type->inner_type->base_type == GLSL_TYPE_FLOAT)
+						{
+							ralloc_asprintf_append(buffer, "highp ");
+						}
+						else if (!GDefaultPrecisionIsHalf && var->type->inner_type->base_type == GLSL_TYPE_HALF)
+						{
+							ralloc_asprintf_append(buffer, "mediump ");
+						}
+						else // shadow samplers, integer textures etc
+						{
+							ralloc_asprintf_append(buffer, "highp ");
+						}
+					}
+					else if (GDefaultPrecisionIsHalf && (var->type->base_type == GLSL_TYPE_FLOAT || (var->type->is_array() && var->type->element_type()->base_type == GLSL_TYPE_FLOAT)))
 					{
 						ralloc_asprintf_append(buffer, "highp ");
 					}
-					else if (!GDefaultPrecisionIsHalf && var->type->inner_type->base_type == GLSL_TYPE_HALF)
+					else if (!GDefaultPrecisionIsHalf && (var->type->base_type == GLSL_TYPE_HALF || (var->type->is_array() && var->type->element_type()->base_type == GLSL_TYPE_HALF)))
 					{
 						ralloc_asprintf_append(buffer, "mediump ");
 					}
-					else // shadow samplers, integer textures etc
+					else if (var->type->is_integer())
 					{
 						ralloc_asprintf_append(buffer, "highp ");
 					}
 				}
-				else if (GDefaultPrecisionIsHalf && (var->type->base_type == GLSL_TYPE_FLOAT || (var->type->is_array() && var->type->element_type()->base_type == GLSL_TYPE_FLOAT)))
-				{
-					ralloc_asprintf_append(buffer, "highp ");
-				}
-				else if (!GDefaultPrecisionIsHalf && (var->type->base_type == GLSL_TYPE_HALF || (var->type->is_array() && var->type->element_type()->base_type == GLSL_TYPE_HALF)))
-				{
-					ralloc_asprintf_append(buffer, "mediump ");
-				}
-				else if (var->type->is_integer())
-				{
-					ralloc_asprintf_append(buffer, "highp ");
-				}
 
 				if (bGenerateLayoutLocations && var->explicit_location)
 				{
@@ -2157,13 +2163,20 @@ class ir_gen_glsl_visitor : public ir_visitor
 
 			if (s->length == 0)
 			{
-				ralloc_asprintf_append(buffer, "\thighp float glsl_doesnt_like_empty_structs;\n");
+				if (bIsES)
+				{
+					ralloc_asprintf_append(buffer, "\thighp float glsl_doesnt_like_empty_structs;\n");
+				}
+				else
+				{
+					ralloc_asprintf_append(buffer, "\tfloat glsl_doesnt_like_empty_structs;\n");
+				}
 			}
 			else
 			{
 				for (unsigned j = 0; j < s->length; j++)
 				{
-					ralloc_asprintf_append(buffer, "\t%s ", (state->language_version == 310) ? "highp" : "") ;
+					ralloc_asprintf_append(buffer, "\t%s ", (state->language_version == 310 && bIsES) ? "highp" : "") ;
 					print_type_pre(s->fields.structure[j].type);
 					ralloc_asprintf_append(buffer, " %s", s->fields.structure[j].name);
 					print_type_post(s->fields.structure[j].type);
@@ -2204,7 +2217,7 @@ class ir_gen_glsl_visitor : public ir_visitor
 					{
 						for (unsigned j = 0; j < type->length; j++)
 						{
-							ralloc_asprintf_append(buffer, "\t%s",  (state->language_version == 310) ? "highp" : "");
+							ralloc_asprintf_append(buffer, "\t%s",  (state->language_version == 310 && bIsES) ? "highp" : "");
 							print_type_pre(type->fields.structure[j].type);
 							ralloc_asprintf_append(buffer, " %s", type->fields.structure[j].name);
 							print_type_post(type->fields.structure[j].type);
@@ -2225,7 +2238,7 @@ class ir_gen_glsl_visitor : public ir_visitor
 						//EHart - name-mangle variables to prevent colliding names
 						ralloc_asprintf_append(buffer, "#define %s %s%s\n", var->name, var->name, block_name);
 
-						ralloc_asprintf_append(buffer, "\t%s", (state->language_version == 310) ? "highp " : "");
+						ralloc_asprintf_append(buffer, "\t%s", (state->language_version == 310 && bIsES) ? "highp " : "");
 						print_type_pre(var->type);
 						ralloc_asprintf_append(buffer, " %s", var->name);
 						print_type_post(var->type);
@@ -2870,7 +2883,7 @@ public:
 			ralloc_asprintf_append(buffer, "#endif\n");
 		}
 
-		if ((state->language_version == 310) && (ShaderTarget == fragment_shader))
+		if ((state->language_version == 310) && (ShaderTarget == fragment_shader) && bIsES)
 		{
 			ralloc_asprintf_append(buffer, "precision %s float;\n", "highp");
 			ralloc_asprintf_append(buffer, "precision %s int;\n", "highp");
diff --git a/Engine/Source/Developer/ShaderFormatOpenGL/Private/ShaderFormatOpenGL.cpp b/Engine/Source/Developer/ShaderFormatOpenGL/Private/ShaderFormatOpenGL.cpp
--- a/Engine/Source/Developer/ShaderFormatOpenGL/Private/ShaderFormatOpenGL.cpp
+++ b/Engine/Source/Developer/ShaderFormatOpenGL/Private/ShaderFormatOpenGL.cpp
@@ -26,7 +26,7 @@ class FShaderFormatGLSL : public IShaderFormat
 		/** Version for shader format, this becomes part of the DDC key. */
 		UE_SHADER_GLSL_150_VER = 26,
 		UE_SHADER_GLSL_150_MAC_VER = 6,
-		UE_SHADER_GLSL_430_VER = 57,
+		UE_SHADER_GLSL_430_VER = 58,
 		UE_SHADER_GLSL_ES2_VER = 15,
 		UE_SHADER_GLSL_150ES2_VER = 18,
 		UE_SHADER_GLSL_150ES2NOUB_VER = 2,
