DisjointTimerQueries off by default, detect IntelHD GPUs, correction for ES 3.1 shaders
#ue4
#android

diff --git a/Engine/Config/BaseDeviceProfiles.ini b/Engine/Config/BaseDeviceProfiles.ini
--- a/Engine/Config/BaseDeviceProfiles.ini
+++ b/Engine/Config/BaseDeviceProfiles.ini
@@ -242,6 +242,15 @@ BaseProfileName=Android
 DeviceType=Android
 BaseProfileName=Android
 
+[Android_IntelHD DeviceProfile]
+DeviceType=Android
+BaseProfileName=Android_Low
+
+[Android_IntelHD_ES3 DeviceProfile]
+DeviceType=Android
+BaseProfileName=Android
++CVars=r.BloomQuality=1
+
 [PS4 DeviceProfile]
 DeviceType=PS4
 BaseProfileName=
diff --git a/Engine/Plugins/Runtime/AndroidDeviceProfileSelector/Source/AndroidDeviceProfileSelector/Private/AndroidDeviceProfileSelectorModule.cpp b/Engine/Plugins/Runtime/AndroidDeviceProfileSelector/Source/AndroidDeviceProfileSelector/Private/AndroidDeviceProfileSelectorModule.cpp
--- a/Engine/Plugins/Runtime/AndroidDeviceProfileSelector/Source/AndroidDeviceProfileSelector/Private/AndroidDeviceProfileSelectorModule.cpp
+++ b/Engine/Plugins/Runtime/AndroidDeviceProfileSelector/Source/AndroidDeviceProfileSelector/Private/AndroidDeviceProfileSelectorModule.cpp
@@ -90,6 +90,19 @@ FString const FAndroidDeviceProfileSelectorModule::GetRuntimeDeviceProfileName()
 			ProfileName = TEXT("Android_Tegra4");
 		}
 	}
+	else if (GPUFamily.StartsWith(TEXT("Intel(R) HD Graphics")))
+	{
+		FString GLVersion = FAndroidMisc::GetGLVersion();
+
+		if (GLVersion.StartsWith(TEXT("OpenGL ES 3.")))
+		{
+			ProfileName = TEXT("Android_IntelHD_ES3");
+		}
+		else if (GLVersion.StartsWith(TEXT("OpenGL ES 2.")))
+		{
+			ProfileName = TEXT("Android_IntelHD");
+		}
+	}
 
 
 	UE_LOG(LogAndroid, Log, TEXT("Selected Device Profile: [%s]"), *ProfileName);
diff --git a/Engine/Source/Developer/ShaderFormatOpenGL/Private/GlslBackend.cpp b/Engine/Source/Developer/ShaderFormatOpenGL/Private/GlslBackend.cpp
--- a/Engine/Source/Developer/ShaderFormatOpenGL/Private/GlslBackend.cpp
+++ b/Engine/Source/Developer/ShaderFormatOpenGL/Private/GlslBackend.cpp
@@ -497,6 +497,7 @@ class ir_gen_glsl_visitor : public ir_visitor
 	// Code generation flags
 	bool bIsES;
 	bool bEmitPrecision;
+	bool bIsES31;
 	_mesa_glsl_parser_targets ShaderTarget;
 
 	bool bGenerateLayoutLocations;
@@ -900,9 +901,11 @@ class ir_gen_glsl_visitor : public ir_visitor
 			}
 			else if (var->type->is_image())
 			{
+				const bool bSingleComp = (var->type->inner_type->vector_elements == 1);
 				const char * const coherent_str[] = { "", "coherent " };
 				const char * const writeonly_str[] = { "", "writeonly " };
-				const char * const type_str[] = { "32ui", "32i", "16f", "32f" };
+				const char * const type_str[] = { "32ui", "32i", "16f", (bIsES31 && !bSingleComp) ? "16f" : "32f" };
+				const char * const comp_str = bSingleComp ? "r" : "rgba";
 				const int writeonly = var->image_write && !(var->image_read);
 
 				check( var->type->inner_type->base_type >= GLSL_TYPE_UINT &&
@@ -922,7 +925,8 @@ class ir_gen_glsl_visitor : public ir_visitor
 					//should check here on base type
 					ralloc_asprintf_append(
 						buffer,
-						"layout(r%s,binding=%d) ",
+						"layout(%s%s,binding=%d) ",
+						comp_str,
 						type_str[var->type->inner_type->base_type],
 						var->location
 						);
@@ -932,7 +936,8 @@ class ir_gen_glsl_visitor : public ir_visitor
 					//should check here on base type
 					ralloc_asprintf_append(
 						buffer,
-						"layout(r%s) ",
+						"layout(%s%s) ",
+						comp_str,
 						type_str[var->type->inner_type->base_type]
 						);
 				}
@@ -2819,10 +2824,11 @@ class ir_gen_glsl_visitor : public ir_visitor
 public:
 
 	/** Constructor. */
-	ir_gen_glsl_visitor(bool bInIsES, bool bInEmitPrecision, _mesa_glsl_parser_targets InShaderTarget, bool bInGenerateLayoutLocations)
+	ir_gen_glsl_visitor(bool bInIsES, bool bInEmitPrecision, bool bInIsES31, _mesa_glsl_parser_targets InShaderTarget, bool bInGenerateLayoutLocations)
 		: early_depth_stencil(false)
 		, bIsES(bInIsES)
 		, bEmitPrecision(bInEmitPrecision)
+		, bIsES31(bInIsES31)
 		, ShaderTarget(InShaderTarget)
 		, bGenerateLayoutLocations(bInGenerateLayoutLocations)
 		, buffer(0)
@@ -3115,7 +3121,7 @@ char* FGlslCodeBackend::GenerateCode(exec_list* ir, _mesa_glsl_parse_state* stat
 	const bool bGroupFlattenedUBs = ((HlslCompileFlags & HLSLCC_GroupFlattenedUniformBuffers) == HLSLCC_GroupFlattenedUniformBuffers);
 	const bool bGenerateLayoutLocations = state->bGenerateLayoutLocations;
 	const bool bEmitPrecision = (Target == HCT_FeatureLevelES2 || Target == HCT_FeatureLevelES3_1 || Target == HCT_FeatureLevelES3_1Ext);
-	ir_gen_glsl_visitor visitor(state->bGenerateES, bEmitPrecision, state->target, bGenerateLayoutLocations);
+	ir_gen_glsl_visitor visitor(state->bGenerateES, bEmitPrecision, (Target == HCT_FeatureLevelES3_1Ext), state->target, bGenerateLayoutLocations);
 	const char* code = visitor.run(ir, state, bGroupFlattenedUBs);
 	return _strdup(code);
 }
diff --git a/Engine/Source/Developer/ShaderFormatOpenGL/Private/ShaderFormatOpenGL.cpp b/Engine/Source/Developer/ShaderFormatOpenGL/Private/ShaderFormatOpenGL.cpp
--- a/Engine/Source/Developer/ShaderFormatOpenGL/Private/ShaderFormatOpenGL.cpp
+++ b/Engine/Source/Developer/ShaderFormatOpenGL/Private/ShaderFormatOpenGL.cpp
@@ -30,10 +30,10 @@ class FShaderFormatGLSL : public IShaderFormat
 		UE_SHADER_GLSL_ES2_VER = 15,
 		UE_SHADER_GLSL_150ES2_VER = 18,
 		UE_SHADER_GLSL_150ES2NOUB_VER = 2,
-		UE_SHADER_GLSL_150ES3_1_VER = 1,
+		UE_SHADER_GLSL_150ES3_1_VER = 2,
 		UE_SHADER_GLSL_ES2_VER_WEBGL  = 16,
 		UE_SHADER_GLSL_ES2_IOS_VER  = 5,
-		UE_SHADER_GLSL_310_ES_EXT_VER = 3,
+		UE_SHADER_GLSL_310_ES_EXT_VER = 4,
 	}; 
 
 	void CheckFormat(FName Format) const
diff --git a/Engine/Source/Runtime/OpenGLDrv/Private/Android/AndroidES31OpenGL.h b/Engine/Source/Runtime/OpenGLDrv/Private/Android/AndroidES31OpenGL.h
--- a/Engine/Source/Runtime/OpenGLDrv/Private/Android/AndroidES31OpenGL.h
+++ b/Engine/Source/Runtime/OpenGLDrv/Private/Android/AndroidES31OpenGL.h
@@ -338,7 +338,7 @@ struct FAndroidES31OpenGL : public FOpenGLES31
 		glGetQueryObjectuiv(QueryId, QueryName, OutResult);
 	}
 
-	static FORCEINLINE void GetQueryObject(GLuint QueryId, EQueryMode QueryMode, uint64* OutResult)
+	static FORCEINLINE void GetQueryObject(GLuint QueryId, EQueryMode QueryMode, GLuint64* OutResult)
 	{
 		GLenum QueryName = (QueryMode == QM_Result) ? GL_QUERY_RESULT_EXT : GL_QUERY_RESULT_AVAILABLE_EXT;
 		GLuint64 Result = 0;
diff --git a/Engine/Source/Runtime/OpenGLDrv/Private/OpenGLES31.cpp b/Engine/Source/Runtime/OpenGLDrv/Private/OpenGLES31.cpp
--- a/Engine/Source/Runtime/OpenGLDrv/Private/OpenGLES31.cpp
+++ b/Engine/Source/Runtime/OpenGLDrv/Private/OpenGLES31.cpp
@@ -10,7 +10,7 @@
 
 static TAutoConsoleVariable<int32> CVarDisjointTimerQueries(
 	TEXT("r.DisjointTimerQueries"),
-	1,
+	0,
 	TEXT("If set to 1, allows GPU time to be measured (e.g. STAT UNIT). It defaults to 0 because some devices supports it but very slowly."),
 	ECVF_RenderThreadSafe);
 
@@ -208,6 +208,13 @@ void FOpenGLES31::ProcessQueryGLInt()
 			GET_GL_INT(GL_MAX_TESS_CONTROL_TEXTURE_IMAGE_UNITS_EXT, 0, MaxHullTextureImageUnits);
 			GET_GL_INT(GL_MAX_TESS_EVALUATION_TEXTURE_IMAGE_UNITS_EXT, 0, MaxDomainTextureImageUnits);
 		}
+		else
+		{
+			MaxHullUniformComponents = 0;
+			MaxDomainUniformComponents = 0;
+			MaxHullTextureImageUnits = 0;
+			MaxDomainTextureImageUnits = 0;
+		}
 	}
 	
 	// No timestamps
