UE4 - Force rebuild SCW

diff --git a/Engine/Source/Developer/ShaderFormatOpenGL/Private/ShaderFormatOpenGL.cpp b/Engine/Source/Developer/ShaderFormatOpenGL/Private/ShaderFormatOpenGL.cpp
--- a/Engine/Source/Developer/ShaderFormatOpenGL/Private/ShaderFormatOpenGL.cpp
+++ b/Engine/Source/Developer/ShaderFormatOpenGL/Private/ShaderFormatOpenGL.cpp
@@ -1,4 +1,5 @@
 // Copyright 1998-2015 Epic Games, Inc. All Rights Reserved.
+//
 
 #include "ShaderFormatOpenGL.h"
 #include "Core.h"
