UE4 - hlslcc - Common header reader

diff --git a/Engine/Source/Developer/ShaderCompilerCommon/Private/ShaderCompilerCommon.cpp b/Engine/Source/Developer/ShaderCompilerCommon/Private/ShaderCompilerCommon.cpp
--- a/Engine/Source/Developer/ShaderCompilerCommon/Private/ShaderCompilerCommon.cpp
+++ b/Engine/Source/Developer/ShaderCompilerCommon/Private/ShaderCompilerCommon.cpp
@@ -3,6 +3,8 @@
 
 #include "ShaderCompilerCommon.h"
 #include "ModuleManager.h"
+#include "CrossCompilerCommon.h"
+#include "TypeHash.h"
 
 
 IMPLEMENT_MODULE(FDefaultModuleImpl, ShaderCompilerCommon);
@@ -273,30 +275,838 @@ bool RemoveUniformBuffersFromSource(FString& SourceCode)
 }
 
 
-FString CreateCrossCompilerBatchFileContents(const FString& ShaderFile, const FString& OutputFile, const FString& FrequencySwitch, const FString& EntryPoint, const FString& VersionSwitch, const FString& ExtraArguments)
+namespace CrossCompiler
 {
-	FString BatchFile;
-	if (PLATFORM_MAC)
+	FString CreateBatchFileContents(const FString& ShaderFile, const FString& OutputFile, const FString& FrequencySwitch, const FString& EntryPoint, const FString& VersionSwitch, const FString& ExtraArguments)
 	{
-		BatchFile = FPaths::RootDir() / FString::Printf(TEXT("Engine/Source/ThirdParty/hlslcc/hlslcc/bin/Mac/hlslcc_64 %s -o=%s %s -entry=%s %s %s"), *ShaderFile, *OutputFile, *FrequencySwitch, *EntryPoint, *VersionSwitch, *ExtraArguments);
+		FString BatchFile;
+		if (PLATFORM_MAC)
+		{
+			BatchFile = FPaths::RootDir() / FString::Printf(TEXT("Engine/Source/ThirdParty/hlslcc/hlslcc/bin/Mac/hlslcc_64 %s -o=%s %s -entry=%s %s %s"), *ShaderFile, *OutputFile, *FrequencySwitch, *EntryPoint, *VersionSwitch, *ExtraArguments);
+		}
+		else if (PLATFORM_LINUX)
+		{
+			BatchFile = FPaths::RootDir() / FString::Printf(TEXT("Engine/Binaries/Linux/CrossCompilerTool %s -o=%s %s -entry=%s %s %s"), *ShaderFile, *OutputFile, *FrequencySwitch, *EntryPoint, *VersionSwitch, *ExtraArguments);
+		}
+		else if (PLATFORM_WINDOWS)
+		{
+			BatchFile = TEXT("@echo off");
+			BatchFile += TEXT("\nif defined ue.hlslcc GOTO DONE\nset ue.hlslcc=");
+			BatchFile += FPaths::RootDir() / TEXT("Engine\\Binaries\\Win64\\CrossCompilerTool.exe");
+			BatchFile += TEXT("\n\n:DONE\n%ue.hlslcc% ");
+			BatchFile += FString::Printf(TEXT("\"%s\" -o=\"%s\" %s -entry=%s %s %s"), *ShaderFile, *OutputFile, *FrequencySwitch, *EntryPoint, *VersionSwitch, *ExtraArguments);
+			BatchFile += TEXT("\npause\n");
+		}
+		else
+		{
+			checkf(false, TEXT("CreateCrossCompilerBatchFileContents: unsupported platform!"));
+		}
+
+		return BatchFile;
+	}
+
+
+	/**
+	 * Parse an error emitted by the HLSL cross-compiler.
+	 * @param OutErrors - Array into which compiler errors may be added.
+	 * @param InLine - A line from the compile log.
+	 */
+	void ParseHlslccError(TArray<FShaderCompilerError>& OutErrors, const FString& InLine)
+	{
+		const TCHAR* p = *InLine;
+		FShaderCompilerError* Error = new(OutErrors) FShaderCompilerError();
+
+		// Copy the filename.
+		while (*p && *p != TEXT('(')) { Error->ErrorFile += (*p++); }
+		Error->ErrorFile = GetRelativeShaderFilename(Error->ErrorFile);
+		p++;
+
+		// Parse the line number.
+		int32 LineNumber = 0;
+		while (*p && *p >= TEXT('0') && *p <= TEXT('9'))
+		{
+			LineNumber = 10 * LineNumber + (*p++ - TEXT('0'));
+		}
+		Error->ErrorLineString = *FString::Printf(TEXT("%d"), LineNumber);
+
+		// Skip to the warning message.
+		while (*p && (*p == TEXT(')') || *p == TEXT(':') || *p == TEXT(' ') || *p == TEXT('\t'))) { p++; }
+		Error->StrippedErrorMessage = p;
+	}
+
+	static inline bool ParseIdentifier(const ANSICHAR*& Str, FString& OutStr)
+	{
+		OutStr = TEXT("");
+		FString Result;
+		while ((*Str >= 'A' && *Str <= 'Z')
+			|| (*Str >= 'a' && *Str <= 'z')
+			|| (*Str >= '0' && *Str <= '9')
+			|| *Str == '_')
+		{
+			OutStr += (TCHAR)*Str;
+			++Str;
+		}
+
+		return OutStr.Len() > 0;
+	}
+
+	static FORCEINLINE bool Match(const ANSICHAR*& Str, ANSICHAR Char)
+	{
+		if (*Str == Char)
+		{
+			++Str;
+			return true;
+		}
+
+		return false;
 	}
-	else if (PLATFORM_LINUX)
+
+	template <typename T>
+	static bool ParseIntegerNumber(const ANSICHAR*& Str, T& OutNum)
 	{
-		BatchFile = FPaths::RootDir() / FString::Printf(TEXT("Engine/Binaries/Linux/CrossCompilerTool %s -o=%s %s -entry=%s %s %s"), *ShaderFile, *OutputFile, *FrequencySwitch, *EntryPoint, *VersionSwitch, *ExtraArguments);
+		auto* OriginalStr = Str;
+		OutNum = 0;
+		while (*Str >= '0' && *Str <= '9')
+		{
+			OutNum = OutNum * 10 + *Str++ - '0';
+		}
+
+		return Str != OriginalStr;
+	}
+
+	static bool ParseSignedNumber(const ANSICHAR*& Str, int32& OutNum)
+	{
+		int32 Sign = Match(Str, '-') ? -1 : 1;
+		uint32 Num = 0;
+		if (ParseIntegerNumber(Str, Num))
+		{
+			OutNum = Sign * (int32)Num;
+			return true;
+		}
+
+		return false;
 	}
-	else if (PLATFORM_WINDOWS)
+
+	/** Map shader frequency -> string for messages. */
+	static const TCHAR* FrequencyStringTable[] =
 	{
-		BatchFile = TEXT("@echo off");
-		BatchFile += TEXT("\nif defined ue.hlslcc GOTO DONE\nset ue.hlslcc=");
-		BatchFile += FPaths::RootDir() / TEXT("Engine\\Binaries\\Win64\\CrossCompilerTool.exe");
-		BatchFile += TEXT("\n\n:DONE\n%ue.hlslcc% ");
-		BatchFile += FString::Printf(TEXT("\"%s\" -o=\"%s\" %s -entry=%s %s %s"), *ShaderFile, *OutputFile, *FrequencySwitch, *EntryPoint, *VersionSwitch, *ExtraArguments);
-		BatchFile += TEXT("\npause\n");
+		TEXT("Vertex"),
+		TEXT("Hull"),
+		TEXT("Domain"),
+		TEXT("Pixel"),
+		TEXT("Geometry"),
+		TEXT("Compute")
+	};
+
+	/** Compile time check to verify that the GL mapping tables are up-to-date. */
+	static_assert(SF_NumFrequencies == ARRAY_COUNT(FrequencyStringTable), "NumFrequencies changed. Please update tables.");
+
+	const TCHAR* GetFrequencyName(EShaderFrequency Frequency)
+	{
+		check((int32)Frequency >= 0 && Frequency < SF_NumFrequencies);
+		return FrequencyStringTable[Frequency];
 	}
-	else
+
+	FHlslccHeader::FHlslccHeader() :
+		Name(TEXT(""))
 	{
-		checkf(false, TEXT("CreateCrossCompilerBatchFileContents: unsupported platform!"));
+		NumThreads[0] = NumThreads[1] = NumThreads[2] = 0;
 	}
 
-	return BatchFile;
+	bool FHlslccHeader::Read(const ANSICHAR*& ShaderSource, int32 SourceLen)
+	{
+#define DEF_PREFIX_STR(Str) \
+		static const ANSICHAR* Str##Prefix = "// @" #Str ": "; \
+		static const int32 Str##PrefixLen = FCStringAnsi::Strlen(Str##Prefix)
+		DEF_PREFIX_STR(Inputs);
+		DEF_PREFIX_STR(Outputs);
+		DEF_PREFIX_STR(UniformBlocks);
+		DEF_PREFIX_STR(Uniforms);
+		DEF_PREFIX_STR(PackedGlobals);
+		DEF_PREFIX_STR(PackedUB);
+		DEF_PREFIX_STR(PackedUBCopies);
+		DEF_PREFIX_STR(PackedUBGlobalCopies);
+		DEF_PREFIX_STR(Samplers);
+		DEF_PREFIX_STR(UAVs);
+		DEF_PREFIX_STR(SamplerStates);
+		DEF_PREFIX_STR(NumThreads);
+#undef DEF_PREFIX_STR
+
+		// Skip any comments that come before the signature.
+		while (FCStringAnsi::Strncmp(ShaderSource, "//", 2) == 0 &&
+			FCStringAnsi::Strncmp(ShaderSource + 2, " !", 2) != 0 &&
+			FCStringAnsi::Strncmp(ShaderSource + 2, " @", 2) != 0)
+		{
+			ShaderSource += 2;
+			while (*ShaderSource && *ShaderSource++ != '\n')
+			{
+				// Do nothing
+			}
+		}
+
+		// Read shader name if any
+		ANSICHAR const* ShaderName = nullptr;
+		uint32 ShaderNameLen = 0;
+		while (FCStringAnsi::Strncmp(ShaderSource, "// !", 4) == 0)
+		{
+			ShaderName = ShaderSource;
+			while (*ShaderSource && *ShaderSource != '\n')
+			{
+				Name += (TCHAR)*ShaderSource;
+				++ShaderSource;
+			}
+		}
+
+		if (FCStringAnsi::Strncmp(ShaderSource, InputsPrefix, InputsPrefixLen) == 0)
+		{
+			ShaderSource += InputsPrefixLen;
+
+			if (!ReadInOut(ShaderSource, Inputs))
+			{
+				return false;
+			}
+		}
+
+		if (FCStringAnsi::Strncmp(ShaderSource, OutputsPrefix, OutputsPrefixLen) == 0)
+		{
+			ShaderSource += OutputsPrefixLen;
+
+			if (!ReadInOut(ShaderSource, Outputs))
+			{
+				return false;
+			}
+		}
+
+		if (FCStringAnsi::Strncmp(ShaderSource, UniformBlocksPrefix, UniformBlocksPrefixLen) == 0)
+		{
+			ShaderSource += UniformBlocksPrefixLen;
+
+			while (*ShaderSource && *ShaderSource != '\n')
+			{
+				FAttribute UniformBlock;
+				if (!ParseIdentifier(ShaderSource, UniformBlock.Name))
+				{
+					return false;
+				}
+
+				if (!Match(ShaderSource, '('))
+				{
+					return false;
+				}
+				
+				if (!ParseIntegerNumber(ShaderSource, UniformBlock.Index))
+				{
+					return false;
+				}
+
+				if (!Match(ShaderSource, ')'))
+				{
+					return false;
+				}
+
+				UniformBlocks.Add(UniformBlock);
+
+				if (Match(ShaderSource, '\n'))
+				{
+					break;
+				}
+
+				if (Match(ShaderSource, ','))
+				{
+					continue;
+				}
+			
+				//#todo-rco: Need a log here
+				//UE_LOG(ShaderCompilerCommon, Warning, TEXT("Invalid char '%c'"), *ShaderSource);
+				return false;
+			}
+		}
+
+		if (FCStringAnsi::Strncmp(ShaderSource, UniformsPrefix, UniformsPrefixLen) == 0)
+		{
+			// @todo-mobile: Will we ever need to support this code path?
+			check(0);
+			return false;
+/*
+			ShaderSource += UniformsPrefixLen;
+
+			while (*ShaderSource && *ShaderSource != '\n')
+			{
+				uint16 ArrayIndex = 0;
+				uint16 Offset = 0;
+				uint16 NumComponents = 0;
+
+				FString ParameterName = ParseIdentifier(ShaderSource);
+				verify(ParameterName.Len() > 0);
+				verify(Match(ShaderSource, '('));
+				ArrayIndex = ParseNumber(ShaderSource);
+				verify(Match(ShaderSource, ':'));
+				Offset = ParseNumber(ShaderSource);
+				verify(Match(ShaderSource, ':'));
+				NumComponents = ParseNumber(ShaderSource);
+				verify(Match(ShaderSource, ')'));
+
+				ParameterMap.AddParameterAllocation(
+					*ParameterName,
+					ArrayIndex,
+					Offset * BytesPerComponent,
+					NumComponents * BytesPerComponent
+					);
+
+				if (ArrayIndex < OGL_NUM_PACKED_UNIFORM_ARRAYS)
+				{
+					PackedUniformSize[ArrayIndex] = FMath::Max<uint16>(
+						PackedUniformSize[ArrayIndex],
+						BytesPerComponent * (Offset + NumComponents)
+						);
+				}
+
+				// Skip the comma.
+				if (Match(ShaderSource, '\n'))
+				{
+					break;
+				}
+
+				verify(Match(ShaderSource, ','));
+			}
+
+			Match(ShaderSource, '\n');
+*/
+		}
+
+		// @PackedGlobals: Global0(h:0,1),Global1(h:4,1),Global2(h:8,1)
+		if (FCStringAnsi::Strncmp(ShaderSource, PackedGlobalsPrefix, PackedGlobalsPrefixLen) == 0)
+		{
+			ShaderSource += PackedGlobalsPrefixLen;
+			while (*ShaderSource && *ShaderSource != '\n')
+			{
+				FPackedGlobal PackedGlobal;
+				if (!ParseIdentifier(ShaderSource, PackedGlobal.Name))
+				{
+					return false;
+				}
+
+				if (!Match(ShaderSource, '('))
+				{
+					return false;
+				}
+
+				PackedGlobal.PackedType = *ShaderSource++;
+
+				if (!Match(ShaderSource, ':'))
+				{
+					return false;
+				}
+
+				if (!ParseIntegerNumber(ShaderSource, PackedGlobal.Offset))
+				{
+					return false;
+				}
+
+				if (!Match(ShaderSource, ','))
+				{
+					return false;
+				}
+
+				if (!ParseIntegerNumber(ShaderSource, PackedGlobal.Count))
+				{
+					return false;
+				}
+
+				if (!Match(ShaderSource, ')'))
+				{
+					return false;
+				}
+
+				PackedGlobals.Add(PackedGlobal);
+
+				// Break if EOL
+				if (Match(ShaderSource, '\n'))
+				{
+					break;
+				}
+
+				// Has to be a comma!
+				if (Match(ShaderSource, ','))
+				{
+					continue;
+				}
+
+				//#todo-rco: Need a log here
+				//UE_LOG(ShaderCompilerCommon, Warning, TEXT("Invalid char '%c'"), *ShaderSource);
+				return false;
+			}
+		}
+
+		// Packed Uniform Buffers (Multiple lines)
+		// @PackedUB: CBuffer(0): CBMember0(0,1),CBMember1(1,1)
+		while (FCStringAnsi::Strncmp(ShaderSource, PackedUBPrefix, PackedUBPrefixLen) == 0)
+		{
+			ShaderSource += PackedUBPrefixLen;
+
+			FPackedUB PackedUB;
+
+			if (!ParseIdentifier(ShaderSource, PackedUB.Attribute.Name))
+			{
+				return false;
+			}
+
+			if (!Match(ShaderSource, '('))
+			{
+				return false;
+			}
+			
+			if (!ParseIntegerNumber(ShaderSource, PackedUB.Attribute.Index))
+			{
+				return false;
+			}
+
+			if (!Match(ShaderSource, ')'))
+			{
+				return false;
+			}
+
+			if (!Match(ShaderSource, ':'))
+			{
+				return false;
+			}
+
+			if (!Match(ShaderSource, ' '))
+			{
+				return false;
+			}
+
+			while (*ShaderSource && *ShaderSource != '\n')
+			{
+				FPackedUB::FMember Member;
+				ParseIdentifier(ShaderSource, Member.Name);
+				if (!Match(ShaderSource, '('))
+				{
+					return false;
+				}
+
+				if (!ParseIntegerNumber(ShaderSource, Member.Offset))
+				{
+					return false;
+				}
+				
+				if (!Match(ShaderSource, ','))
+				{
+					return false;
+				}
+
+				if (!ParseIntegerNumber(ShaderSource, Member.Count))
+				{
+					return false;
+				}
+
+				if (!Match(ShaderSource, ')'))
+				{
+					return false;
+				}
+
+				PackedUB.Members.Add(Member);
+
+				// Break if EOL
+				if (Match(ShaderSource, '\n'))
+				{
+					break;
+				}
+
+				// Has to be a comma!
+				if (Match(ShaderSource, ','))
+				{
+					continue;
+				}
+
+				//#todo-rco: Need a log here
+				//UE_LOG(ShaderCompilerCommon, Warning, TEXT("Invalid char '%c'"), *ShaderSource);
+				return false;
+			}
+
+			PackedUBs.Add(PackedUB);
+		}
+
+		// @PackedUBCopies: 0:0-0:h:0:1,0:1-0:h:4:1,1:0-1:h:0:1
+		if (FCStringAnsi::Strncmp(ShaderSource, PackedUBCopiesPrefix, PackedUBCopiesPrefixLen) == 0)
+		{
+			ShaderSource += PackedUBCopiesPrefixLen;
+			if (!ReadCopies(ShaderSource, false, PackedUBCopies))
+			{
+				return false;
+			}
+		}
+
+		// @PackedUBGlobalCopies: 0:0-h:12:1,0:1-h:16:1,1:0-h:20:1
+		if (FCStringAnsi::Strncmp(ShaderSource, PackedUBGlobalCopiesPrefix, PackedUBGlobalCopiesPrefixLen) == 0)
+		{
+			ShaderSource += PackedUBGlobalCopiesPrefixLen;
+			if (!ReadCopies(ShaderSource, true, PackedUBGlobalCopies))
+			{
+				return false;
+			}
+		}
+
+		if (FCStringAnsi::Strncmp(ShaderSource, SamplersPrefix, SamplersPrefixLen) == 0)
+		{
+			ShaderSource += SamplersPrefixLen;
+
+			while (*ShaderSource && *ShaderSource != '\n')
+			{
+				FSampler Sampler;
+
+				if (!ParseIdentifier(ShaderSource, Sampler.Name))
+				{
+					return false;
+				}
+
+				if (!Match(ShaderSource, '('))
+				{
+					return false;
+				}
+
+				if (!ParseIntegerNumber(ShaderSource, Sampler.Offset))
+				{
+					return false;
+				}
+
+				if (!Match(ShaderSource, ':'))
+				{
+					return false;
+				}
+
+				if (!ParseIntegerNumber(ShaderSource, Sampler.Count))
+				{
+					return false;
+				}
+
+				if (Match(ShaderSource, '['))
+				{
+					// Sampler States
+					do
+					{
+						FString SamplerState;
+						
+						if (!ParseIdentifier(ShaderSource, SamplerState))
+						{
+							return false;
+						}
+
+						Sampler.SamplerStates.Add(SamplerState);
+					}
+					while (Match(ShaderSource, ','));
+
+					if (!Match(ShaderSource, ']'))
+					{
+						return false;
+					}
+				}
+
+				if (!Match(ShaderSource, ')'))
+				{
+					return false;
+				}
+
+				Samplers.Add(Sampler);
+
+				// Break if EOL
+				if (Match(ShaderSource, '\n'))
+				{
+					break;
+				}
+
+				// Has to be a comma!
+				if (Match(ShaderSource, ','))
+				{
+					continue;
+				}
+
+				//#todo-rco: Need a log here
+				//UE_LOG(ShaderCompilerCommon, Warning, TEXT("Invalid char '%c'"), *ShaderSource);
+				return false;
+			}
+		}
+
+		if (FCStringAnsi::Strncmp(ShaderSource, UAVsPrefix, UAVsPrefixLen) == 0)
+		{
+			ShaderSource += UAVsPrefixLen;
+
+			while (*ShaderSource && *ShaderSource != '\n')
+			{
+				FUAV UAV;
+
+				if (!ParseIdentifier(ShaderSource, UAV.Name))
+				{
+					return false;
+				}
+
+				if (!Match(ShaderSource, '('))
+				{
+					return false;
+				}
+
+				if (!ParseIntegerNumber(ShaderSource, UAV.Offset))
+				{
+					return false;
+				}
+
+				if (!Match(ShaderSource, ':'))
+				{
+					return false;
+				}
+
+				if (!ParseIntegerNumber(ShaderSource, UAV.Count))
+				{
+					return false;
+				}
+
+				if (!Match(ShaderSource, ')'))
+				{
+					return false;
+				}
+
+				UAVs.Add(UAV);
+
+				// Break if EOL
+				if (Match(ShaderSource, '\n'))
+				{
+					break;
+				}
+
+				// Has to be a comma!
+				if (Match(ShaderSource, ','))
+				{
+					continue;
+				}
+
+				//#todo-rco: Need a log here
+				//UE_LOG(ShaderCompilerCommon, Warning, TEXT("Invalid char '%c'"), *ShaderSource);
+				return false;
+			}
+		}
+
+		if (FCStringAnsi::Strncmp(ShaderSource, SamplerStatesPrefix, SamplerStatesPrefixLen) == 0)
+		{
+			ShaderSource += SamplerStatesPrefixLen;
+			while (*ShaderSource && *ShaderSource != '\n')
+			{
+				FAttribute SamplerState;
+				if (!ParseIntegerNumber(ShaderSource, SamplerState.Index))
+				{
+					return false;
+				}
+
+				if (!Match(ShaderSource, ':'))
+				{
+					return false;
+				}
+
+				if (!ParseIdentifier(ShaderSource, SamplerState.Name))
+				{
+					return false;
+				}
+
+				SamplerStates.Add(SamplerState);
+			}
+		}
+
+		if (FCStringAnsi::Strncmp(ShaderSource, NumThreadsPrefix, NumThreadsPrefixLen) == 0)
+		{
+			ShaderSource += NumThreadsPrefixLen;
+			if (!ParseIntegerNumber(ShaderSource, NumThreads[0]))
+			{
+				return false;
+			}
+			if (!Match(ShaderSource, ','))
+			{
+				return false;
+			}
+
+			if (!Match(ShaderSource, ' '))
+			{
+				return false;
+			}
+
+			if (!ParseIntegerNumber(ShaderSource, NumThreads[1]))
+			{
+				return false;
+			}
+
+			if (!Match(ShaderSource, ','))
+			{
+				return false;
+			}
+
+			if (!Match(ShaderSource, ' '))
+			{
+				return false;
+			}
+
+			if (!ParseIntegerNumber(ShaderSource, NumThreads[2]))
+			{
+				return false;
+			}
+
+			if (!Match(ShaderSource, '\n'))
+			{
+				return false;
+			}
+		}
+
+		return true;
+	}
+
+	bool FHlslccHeader::ReadCopies(const ANSICHAR*& ShaderSource, bool bGlobals, TArray<FPackedUBCopy>& OutCopies)
+	{
+		while (*ShaderSource && *ShaderSource != '\n')
+		{
+			FPackedUBCopy PackedUBCopy;
+			PackedUBCopy.DestUB = 0;
+
+			if (!ParseIntegerNumber(ShaderSource, PackedUBCopy.SourceUB))
+			{
+				return false;
+			}
+
+			if (!Match(ShaderSource, ':'))
+			{
+				return false;
+			}
+
+			if (!ParseIntegerNumber(ShaderSource, PackedUBCopy.SourceOffset))
+			{
+				return false;
+			}
+
+			if (!Match(ShaderSource, '-'))
+			{
+				return false;
+			}
+
+			if (!bGlobals)
+			{
+				if (!ParseIntegerNumber(ShaderSource, PackedUBCopy.DestUB))
+				{
+					return false;
+				}
+
+				if (!Match(ShaderSource, ':'))
+				{
+					return false;
+				}
+			}
+
+			PackedUBCopy.DestPackedType = *ShaderSource++;
+
+			if (!Match(ShaderSource, ':'))
+			{
+				return false;
+			}
+
+			if (!ParseIntegerNumber(ShaderSource, PackedUBCopy.DestOffset))
+			{
+				return false;
+			}
+
+			if (!Match(ShaderSource, ':'))
+			{
+				return false;
+			}
+
+			if (!ParseIntegerNumber(ShaderSource, PackedUBCopy.Count))
+			{
+				return false;
+			}
+
+			OutCopies.Add(PackedUBCopy);
+
+			// Break if EOL
+			if (Match(ShaderSource, '\n'))
+			{
+				break;
+			}
+
+			// Has to be a comma!
+			if (Match(ShaderSource, ','))
+			{
+				continue;
+			}
+
+			//#todo-rco: Need a log here
+			//UE_LOG(ShaderCompilerCommon, Warning, TEXT("Invalid char '%c'"), *ShaderSource);
+			return false;
+		}
+
+		return true;
+	}
+
+	bool FHlslccHeader::ReadInOut(const ANSICHAR*& ShaderSource, TArray<FInOut>& OutAttributes)
+	{
+		while (*ShaderSource && *ShaderSource != '\n')
+		{
+			FInOut Attribute;
+
+			if (!ParseIdentifier(ShaderSource, Attribute.Type))
+			{
+				return false;
+			}
+
+			if (Match(ShaderSource, '['))
+			{
+				if (!ParseIntegerNumber(ShaderSource, Attribute.ArrayCount))
+				{
+					return false;
+				}
+
+				if (!Match(ShaderSource, ']'))
+				{
+					return false;
+				}
+			}
+			else
+			{
+				Attribute.ArrayCount = 0;
+			}
+
+			if (!Match(ShaderSource, ';'))
+			{
+				return false;
+			}
+
+			if (!ParseSignedNumber(ShaderSource, Attribute.Index))
+			{
+				return false;
+			}
+
+			if (!Match(ShaderSource, ':'))
+			{
+				return false;
+			}
+
+			if (!ParseIdentifier(ShaderSource, Attribute.Name))
+			{
+				return false;
+			}
+
+			OutAttributes.Add(Attribute);
+
+			// Break if EOL
+			if (Match(ShaderSource, '\n'))
+			{
+				break;
+			}
+
+			// Has to be a comma!
+			if (Match(ShaderSource, ','))
+			{
+				continue;
+			}
+
+			//#todo-rco: Need a log here
+			//UE_LOG(ShaderCompilerCommon, Warning, TEXT("Invalid char '%c'"), *ShaderSource);
+			return false;
+		}
+
+		return true;
+	}
 }
diff --git a/Engine/Source/Developer/ShaderCompilerCommon/Public/ShaderCompilerCommon.h b/Engine/Source/Developer/ShaderCompilerCommon/Public/ShaderCompilerCommon.h
--- a/Engine/Source/Developer/ShaderCompilerCommon/Public/ShaderCompilerCommon.h
+++ b/Engine/Source/Developer/ShaderCompilerCommon/Public/ShaderCompilerCommon.h
@@ -4,6 +4,7 @@
 
 #include "Core.h"
 #include "ShaderCore.h"
+#include "CrossCompilerCommon.h"
 
 
 /**
@@ -63,11 +64,104 @@ extern SHADERCOMPILERCOMMON_API int16 GetNumUniformBuffersUsed(const FShaderReso
 extern SHADERCOMPILERCOMMON_API bool RemoveUniformBuffersFromSource(FString& SourceCode);
 
 // Cross compiler support/common functionality
+namespace CrossCompiler
+{
+	extern SHADERCOMPILERCOMMON_API FString CreateBatchFileContents(
+		const FString& ShaderFile,
+		const FString& OutputFile,
+		const FString& FrequencySwitch,
+		const FString& EntryPoint,
+		const FString& VersionSwitch,
+		const FString& ExtraArguments = TEXT(""));
 
-extern SHADERCOMPILERCOMMON_API FString CreateCrossCompilerBatchFileContents(
-											const FString& ShaderFile,
-											const FString& OutputFile,
-											const FString& FrequencySwitch,
-											const FString& EntryPoint,
-											const FString& VersionSwitch,
-											const FString& ExtraArguments = TEXT(""));
+	extern SHADERCOMPILERCOMMON_API void ParseHlslccError(TArray<FShaderCompilerError>& OutErrors, const FString& InLine);
+
+	struct SHADERCOMPILERCOMMON_API FHlslccHeader
+	{
+		FHlslccHeader();
+
+		bool Read(const ANSICHAR*& ShaderSource, int32 SourceLen);
+
+		struct FInOut
+		{
+			FString Type;
+			int32 Index;
+			int32 ArrayCount;
+			FString Name;
+		};
+
+		struct FAttribute
+		{
+			int32 Index;
+			FString Name;
+		};
+
+		struct FPackedGlobal
+		{
+			ANSICHAR PackedType;
+			FString Name;
+			int32 Offset;
+			int32 Count;
+		};
+
+		//struct FUniform
+		//{
+		//};
+
+		struct FPackedUB
+		{
+			FAttribute Attribute;
+			struct FMember
+			{
+				FString Name;
+				int32 Offset;
+				int32 Count;
+			};
+			TArray<FMember> Members;
+		};
+
+		struct FPackedUBCopy
+		{
+			int32 SourceUB;
+			int32 SourceOffset;
+			int32 DestUB;
+			ANSICHAR DestPackedType;
+			int32 DestOffset;
+			int32 Count;
+		};
+
+		struct FSampler
+		{
+			FString Name;
+			int32 Offset;
+			int32 Count;
+			TArray<FString> SamplerStates;
+		};
+
+		struct FUAV
+		{
+			FString Name;
+			int32 Offset;
+			int32 Count;
+		};
+
+		FString Name;
+		TArray<FInOut> Inputs;
+		TArray<FInOut> Outputs;
+		TArray<FAttribute> UniformBlocks;
+		//TArray<FUniform> Uniforms;
+		TArray<FPackedGlobal> PackedGlobals;
+		TArray<FPackedUB> PackedUBs;
+		TArray<FPackedUBCopy> PackedUBCopies;
+		TArray<FPackedUBCopy> PackedUBGlobalCopies;
+		TArray<FSampler> Samplers;
+		TArray<FUAV> UAVs;
+		TArray<FAttribute> SamplerStates;
+		uint32 NumThreads[3];
+
+		static bool ReadInOut(const ANSICHAR*& ShaderSource, TArray<FInOut>& OutAttributes);
+		static bool ReadCopies(const ANSICHAR*& ShaderSource, bool bGlobals, TArray<FPackedUBCopy>& OutCopies);
+	};
+
+	extern SHADERCOMPILERCOMMON_API const TCHAR* GetFrequencyName(EShaderFrequency Frequency);
+}
diff --git a/Engine/Source/Developer/ShaderFormatOpenGL/Private/GlslBackend.cpp b/Engine/Source/Developer/ShaderFormatOpenGL/Private/GlslBackend.cpp
--- a/Engine/Source/Developer/ShaderFormatOpenGL/Private/GlslBackend.cpp
+++ b/Engine/Source/Developer/ShaderFormatOpenGL/Private/GlslBackend.cpp
@@ -2616,7 +2616,7 @@ class ir_gen_glsl_visitor : public ir_visitor
 			}
 			ralloc_asprintf_append(buffer, "%s%s%s%s",
 				need_comma ? "," : "",
-				type_str[type->base_type],
+				type->base_type == GLSL_TYPE_STRUCT ? type->name : type_str[type->base_type],
 				col_str[type->matrix_columns],
 				row_str[type->vector_elements]);
 			if (is_array)
diff --git a/Engine/Source/Developer/ShaderFormatOpenGL/Private/OpenGLShaderCompiler.cpp b/Engine/Source/Developer/ShaderFormatOpenGL/Private/OpenGLShaderCompiler.cpp
--- a/Engine/Source/Developer/ShaderFormatOpenGL/Private/OpenGLShaderCompiler.cpp
+++ b/Engine/Source/Developer/ShaderFormatOpenGL/Private/OpenGLShaderCompiler.cpp
@@ -438,20 +438,7 @@ GLenum GLFrequencyTable[] =
 	
 };
 
-/** Map shader frequency -> string for messages. */
-static const TCHAR* GLFrequencyStringTable[] =
-{
-	TEXT("Vertex"),
-	TEXT("Hull"),
-	TEXT("Domain"),
-	TEXT("Pixel"),
-	TEXT("Geometry"),
-	TEXT("Compute")
-};
-
-/** Compile time check to verify that the GL mapping tables are up-to-date. */
-static_assert(SF_NumFrequencies == ARRAY_COUNT(GLFrequencyTable), "NumFrequencies changed. Please update tables.");
-static_assert(ARRAY_COUNT(GLFrequencyTable) == ARRAY_COUNT(GLFrequencyStringTable), "Frequency table size mismatch.");
+static_assert(ARRAY_COUNT(GLFrequencyTable) == SF_NumFrequencies, "Frequency table size mismatch.");
 
 /**
  * Parse a GLSL error.
@@ -494,51 +481,20 @@ void ParseGlslError(TArray<FShaderCompilerError>& OutErrors, const FString& InLi
 	}
 }
 
-static TArray<ANSICHAR> ParseIdentifierANSI(const ANSICHAR* &Str)
+static TArray<ANSICHAR> ParseIdentifierANSI(const FString& Str)
 {
 	TArray<ANSICHAR> Result;
-	
-	while ((*Str >= 'A' && *Str <= 'Z')
-		   || (*Str >= 'a' && *Str <= 'z')
-		   || (*Str >= '0' && *Str <= '9')
-		   || *Str == '_')
+	Result.Reserve(Str.Len());
+	for (int32 Index = 0; Index < Str.Len(); ++Index)
 	{
-		Result.Add(FChar::ToLower(*Str));
-		++Str;
+		Result.Add(FChar::ToLower((ANSICHAR)Str[Index]));
 	}
 	Result.Add('\0');
-	
-	return Result;
-}
-
-static FString ParseIdentifier(const ANSICHAR* &Str)
-{
-	FString Result;
-
-	while ((*Str >= 'A' && *Str <= 'Z')
-		|| (*Str >= 'a' && *Str <= 'z')
-		|| (*Str >= '0' && *Str <= '9')
-		|| *Str == '_')
-	{
-		Result += *Str;
-		++Str;
-	}
 
 	return Result;
 }
 
-static bool Match(const ANSICHAR* &Str, ANSICHAR Char)
-{
-	if (*Str == Char)
-	{
-		++Str;
-		return true;
-	}
-
-	return false;
-}
-
-static uint32 ParseNumber(const ANSICHAR* &Str)
+static uint32 ParseNumber(const TCHAR* Str)
 {
 	uint32 Num = 0;
 	while (*Str && *Str >= '0' && *Str <= '9')
@@ -562,13 +518,24 @@ static void BuildShaderOutput(
 	GLSLVersion Version
 	)
 {
+	const ANSICHAR* USFSource = InShaderSource;
+	CrossCompiler::FHlslccHeader CCHeader;
+	if (!CCHeader.Read(USFSource, SourceLen))
+	{
+		UE_LOG(LogOpenGLShaderCompiler, Error, TEXT("Bad hlslcc header found"));
+	}
+	
+	if (*USFSource != '#')
+	{
+		UE_LOG(LogOpenGLShaderCompiler, Error, TEXT("Bad hlslcc header found! Missing '#'!"));
+	}
+
 	FOpenGLCodeHeader Header = {0};
-	const ANSICHAR* ShaderSource = InShaderSource;
 	FShaderParameterMap& ParameterMap = ShaderOutput.ParameterMap;
 	EShaderFrequency Frequency = (EShaderFrequency)ShaderOutput.Target.Frequency;
 
 	TBitArray<> UsedUniformBufferSlots;
-	UsedUniformBufferSlots.Init(false,32);
+	UsedUniformBufferSlots.Init(false, 32);
 
 	// Write out the magic markers.
 	Header.GlslMarker = 0x474c534c;
@@ -596,74 +563,24 @@ static void BuildShaderOutput(
 		UE_LOG(LogOpenGLShaderCompiler, Fatal, TEXT("Invalid shader frequency: %d"), (int32)Frequency);
 	}
 
-	#define DEF_PREFIX_STR(Str) \
-		const ANSICHAR* Str##Prefix = "// @" #Str ": "; \
-		const int32 Str##PrefixLen = FCStringAnsi::Strlen(Str##Prefix)
-	DEF_PREFIX_STR(Inputs);
-	DEF_PREFIX_STR(Outputs);
-	DEF_PREFIX_STR(UniformBlocks);
-	DEF_PREFIX_STR(Uniforms);
-	DEF_PREFIX_STR(PackedGlobals);
-	DEF_PREFIX_STR(PackedUB);
-	DEF_PREFIX_STR(PackedUBCopies);
-	DEF_PREFIX_STR(PackedUBGlobalCopies);
-	DEF_PREFIX_STR(Samplers);
-	DEF_PREFIX_STR(UAVs);
-	DEF_PREFIX_STR(SamplerStates);
-	#undef DEF_PREFIX_STR
-
-	// Skip any comments that come before the signature.
-	while (	FCStringAnsi::Strncmp(ShaderSource, "//", 2) == 0 &&
-			FCStringAnsi::Strncmp(ShaderSource, "// @", 4) != 0 )
+	static const FString AttributePrefix = TEXT("in_ATTRIBUTE");
+	static const FString GL_Prefix = TEXT("gl_");
+	for (auto& Input : CCHeader.Inputs)
 	{
-		while (*ShaderSource && *ShaderSource++ != '\n') {}
-	}
-
-	// HLSLCC first prints the list of inputs.
-	if (FCStringAnsi::Strncmp(ShaderSource, InputsPrefix, InputsPrefixLen) == 0)
-	{
-		ShaderSource += InputsPrefixLen;
-
-		const ANSICHAR* AttributePrefix = "in_ATTRIBUTE";
-        const int32 AttributePrefixLen = FCStringAnsi::Strlen(AttributePrefix);
-        while (*ShaderSource && *ShaderSource != '\n')
-        {
-            // Get the location
-            while (*ShaderSource && *ShaderSource++ != ';') {}
-            int32 Location = FCStringAnsi::Atoi(ShaderSource);
-            
-            // Skip the type.
-            while (*ShaderSource && *ShaderSource++ != ':') {}
-            
-            // Only process attributes for vertex shaders.
-            if (Frequency == SF_Vertex && FCStringAnsi::Strncmp(ShaderSource, AttributePrefix, AttributePrefixLen) == 0)
-            {
-                ShaderSource += AttributePrefixLen;
-                uint8 AttributeIndex = ParseNumber(ShaderSource);
-                Header.Bindings.InOutMask |= (1 << AttributeIndex);
-            }
-            // Record user-defined input varyings
-            else if( FCStringAnsi::Strncmp(ShaderSource, "gl_", 3) != 0 )
-            {
-                FOpenGLShaderVarying Var;
-                Var.Location = Location;
-                Var.Varying = ParseIdentifierANSI(ShaderSource);
-                Header.Bindings.InputVaryings.Add(Var);
-            }
-
-            // Skip to the next.
-            while (*ShaderSource && *ShaderSource != ',' && *ShaderSource != '\n')
-            {
-                ShaderSource++;
-            }
-
-            if (Match(ShaderSource, '\n'))
-            {
-                break;
-            }
-
-            verify(Match(ShaderSource, ','));
-        }
+		// Only process attributes for vertex shaders.
+		if (Frequency == SF_Vertex && Input.Name.StartsWith(AttributePrefix))
+		{
+			int32 AttributeIndex = ParseNumber(*Input.Name + AttributePrefix.Len());
+			Header.Bindings.InOutMask |= (1 << AttributeIndex);
+		}
+		// Record user-defined input varyings
+		else if (!Input.Name.StartsWith(GL_Prefix))
+		{
+			FOpenGLShaderVarying Var;
+			Var.Location = Input.Index;
+			Var.Varying = ParseIdentifierANSI(Input.Name);
+			Header.Bindings.InputVaryings.Add(Var);
+		}
 	}
 
 	// Generate vertex attribute remapping table.
@@ -673,7 +590,7 @@ static void BuildShaderOutput(
 		uint32 AttributeMask = Header.Bindings.InOutMask;
 		int32 NextAttributeSlot = 0;
 		Header.Bindings.VertexRemappedMask = 0;
-		for (int32 AttributeIndex = 0; AttributeIndex<16; AttributeIndex++, AttributeMask >>= 1)
+		for (int32 AttributeIndex = 0; AttributeIndex < 16; AttributeIndex++, AttributeMask >>= 1)
 		{
 			if (AttributeMask & 0x1)
 			{
@@ -687,220 +604,69 @@ static void BuildShaderOutput(
 		}
 	}
 
-	// Then the list of outputs.
-	if (FCStringAnsi::Strncmp(ShaderSource, OutputsPrefix, OutputsPrefixLen) == 0)
+	static const FString TargetPrefix = "out_Target";
+	static const FString GL_FragDepth = "gl_FragDepth";
+	for (auto& Output : CCHeader.Outputs)
 	{
-		ShaderSource += OutputsPrefixLen;
-        
-        const ANSICHAR* TargetPrefix = "out_Target";
-        const int32 TargetPrefixLen = FCStringAnsi::Strlen(TargetPrefix);
-        
-		while (*ShaderSource && *ShaderSource != '\n')
+        // Only targets for pixel shaders must be tracked.
+        if (Frequency == SF_Pixel && Output.Name.StartsWith(TargetPrefix))
+        {
+			uint8 TargetIndex = ParseNumber(*Output.Name + TargetPrefix.Len());
+            Header.Bindings.InOutMask |= (1 << TargetIndex);
+        }
+        // Only depth writes for pixel shaders must be tracked.
+        else if (Frequency == SF_Pixel && Output.Name.Equals(GL_FragDepth))
         {
-            // Get the location
-            while (*ShaderSource && *ShaderSource++ != ';') {}
-            int32 Location = FCStringAnsi::Atoi(ShaderSource);
-            
-            // Skip the type.
-            while (*ShaderSource && *ShaderSource++ != ':') {}
-
-            // Only targets for pixel shaders must be tracked.
-            if (Frequency == SF_Pixel && FCStringAnsi::Strncmp(ShaderSource, TargetPrefix, TargetPrefixLen) == 0)
-            {
-                ShaderSource += TargetPrefixLen;
-                uint8 TargetIndex = ParseNumber(ShaderSource);
-                Header.Bindings.InOutMask |= (1 << TargetIndex);
-            }
-            // Only depth writes for pixel shaders must be tracked.
-            else if (Frequency == SF_Pixel && FCStringAnsi::Strcmp(ShaderSource, "gl_FragDepth") == 0)
-            {
-                Header.Bindings.InOutMask |= 0x8000;
-            }
-            // Record user-defined output varyings
-            else if( FCStringAnsi::Strncmp(ShaderSource, "gl_", 3) != 0 )
-            {
-                FOpenGLShaderVarying Var;
-                Var.Location = Location;
-                Var.Varying = ParseIdentifierANSI(ShaderSource);
-                Header.Bindings.OutputVaryings.Add(Var);
-            }
-
-            // Skip to the next.
-            while (*ShaderSource && *ShaderSource != ',' && *ShaderSource != '\n')
-            {
-                ShaderSource++;
-            }
-
-            if (Match(ShaderSource, '\n'))
-            {
-                break;
-            }
-
-            verify(Match(ShaderSource, ','));
+            Header.Bindings.InOutMask |= 0x8000;
+        }
+        // Record user-defined output varyings
+        else if (!Output.Name.StartsWith(GL_Prefix))
+        {
+            FOpenGLShaderVarying Var;
+            Var.Location = Output.Index;
+            Var.Varying = ParseIdentifierANSI(Output.Name);
+            Header.Bindings.OutputVaryings.Add(Var);
         }
 	}
 
 	// Then 'normal' uniform buffers.
-	if (FCStringAnsi::Strncmp(ShaderSource, UniformBlocksPrefix, UniformBlocksPrefixLen) == 0)
+	for (auto& UniformBlock : CCHeader.UniformBlocks)
 	{
-		ShaderSource += UniformBlocksPrefixLen;
-
-		while (*ShaderSource && *ShaderSource != '\n')
-		{
-			FString BufferName = ParseIdentifier(ShaderSource);
-			verify(BufferName.Len() > 0);
-			verify(Match(ShaderSource, '('));
-			uint16 UBIndex = ParseNumber(ShaderSource);
-			check(UBIndex == Header.Bindings.NumUniformBuffers);
-			UsedUniformBufferSlots[UBIndex] = true;
-			verify(Match(ShaderSource, ')'));
-			ParameterMap.AddParameterAllocation(*BufferName, Header.Bindings.NumUniformBuffers++, 0, 0);
-
-			// Skip the comma.
-			if (Match(ShaderSource, '\n'))
-			{
-				break;
-			}
-
-			verify(Match(ShaderSource, ','));
-		}
-
-		Match(ShaderSource, '\n');
+		uint16 UBIndex = UniformBlock.Index;
+		check(UBIndex == Header.Bindings.NumUniformBuffers);
+		UsedUniformBufferSlots[UBIndex] = true;
+		ParameterMap.AddParameterAllocation(*UniformBlock.Name, Header.Bindings.NumUniformBuffers++, 0, 0);
 	}
 
-	// Then uniforms.
 	const uint16 BytesPerComponent = 4;
-/*
-	uint16 PackedUniformSize[OGL_NUM_PACKED_UNIFORM_ARRAYS] = {0};
-	FMemory::Memzero(&PackedUniformSize, sizeof(PackedUniformSize));
-*/
-	if (FCStringAnsi::Strncmp(ShaderSource, UniformsPrefix, UniformsPrefixLen) == 0)
-	{
-		// @todo-mobile: Will we ever need to support this code path?
-		check(0);
-/*
-		ShaderSource += UniformsPrefixLen;
-
-		while (*ShaderSource && *ShaderSource != '\n')
-		{
-			uint16 ArrayIndex = 0;
-			uint16 Offset = 0;
-			uint16 NumComponents = 0;
-
-			FString ParameterName = ParseIdentifier(ShaderSource);
-			verify(ParameterName.Len() > 0);
-			verify(Match(ShaderSource, '('));
-			ArrayIndex = ParseNumber(ShaderSource);
-			verify(Match(ShaderSource, ':'));
-			Offset = ParseNumber(ShaderSource);
-			verify(Match(ShaderSource, ':'));
-			NumComponents = ParseNumber(ShaderSource);
-			verify(Match(ShaderSource, ')'));
-
-			ParameterMap.AddParameterAllocation(
-				*ParameterName,
-				ArrayIndex,
-				Offset * BytesPerComponent,
-				NumComponents * BytesPerComponent
-				);
-
-			if (ArrayIndex < OGL_NUM_PACKED_UNIFORM_ARRAYS)
-			{
-				PackedUniformSize[ArrayIndex] = FMath::Max<uint16>(
-					PackedUniformSize[ArrayIndex],
-					BytesPerComponent * (Offset + NumComponents)
-					);
-			}
-
-			// Skip the comma.
-			if (Match(ShaderSource, '\n'))
-			{
-				break;
-			}
-
-			verify(Match(ShaderSource, ','));
-		}
-
-		Match(ShaderSource, '\n');
-*/
-	}
 
 	// Packed global uniforms
 	TMap<ANSICHAR, uint16> PackedGlobalArraySize;
-	if (FCStringAnsi::Strncmp(ShaderSource, PackedGlobalsPrefix, PackedGlobalsPrefixLen) == 0)
+	for (auto& PackedGlobal : CCHeader.PackedGlobals)
 	{
-		ShaderSource += PackedGlobalsPrefixLen;
-		while (*ShaderSource && *ShaderSource != '\n')
-		{
-			ANSICHAR ArrayIndex = 0;
-			uint16 Offset = 0;
-			uint16 NumComponents = 0;
-
-			FString ParameterName = ParseIdentifier(ShaderSource);
-			verify(ParameterName.Len() > 0);
-			verify(Match(ShaderSource, '('));
-			ArrayIndex = *ShaderSource++;
-			verify(Match(ShaderSource, ':'));
-			Offset = ParseNumber(ShaderSource);
-			verify(Match(ShaderSource, ','));
-			NumComponents = ParseNumber(ShaderSource);
-			verify(Match(ShaderSource, ')'));
-
-			ParameterMap.AddParameterAllocation(
-				*ParameterName,
-				ArrayIndex,
-				Offset * BytesPerComponent,
-				NumComponents * BytesPerComponent
-				);
-
-			uint16& Size = PackedGlobalArraySize.FindOrAdd(ArrayIndex);
-			Size = FMath::Max<uint16>(BytesPerComponent * (Offset + NumComponents), Size);
-
-			if (Match(ShaderSource, '\n'))
-			{
-				break;
-			}
-
-			// Skip the comma.
-			verify(Match(ShaderSource, ','));
-		}
-
-		Match(ShaderSource, '\n');
+		ParameterMap.AddParameterAllocation(
+			*PackedGlobal.Name,
+			PackedGlobal.PackedType,
+			PackedGlobal.Offset * BytesPerComponent,
+			PackedGlobal.Count * BytesPerComponent
+			);
+
+		uint16& Size = PackedGlobalArraySize.FindOrAdd(PackedGlobal.PackedType);
+		Size = FMath::Max<uint16>(BytesPerComponent * (PackedGlobal.Offset + PackedGlobal.Count), Size);
 	}
 
 	// Packed Uniform Buffers
 	TMap<int, TMap<ANSICHAR, uint16> > PackedUniformBuffersSize;
-	while (FCStringAnsi::Strncmp(ShaderSource, PackedUBPrefix, PackedUBPrefixLen) == 0)
+	for (auto& PackedUB : CCHeader.PackedUBs)
 	{
-		ShaderSource += PackedUBPrefixLen;
-		FString BufferName = ParseIdentifier(ShaderSource);
-		verify(BufferName.Len() > 0);
-		verify(Match(ShaderSource, '('));
-		uint16 BufferIndex = ParseNumber(ShaderSource);
-		check(BufferIndex == Header.Bindings.NumUniformBuffers);
-		UsedUniformBufferSlots[BufferIndex] = true;
-		verify(Match(ShaderSource, ')'));
-		ParameterMap.AddParameterAllocation(*BufferName, Header.Bindings.NumUniformBuffers++, 0, 0);
-
-		verify(Match(ShaderSource, ':'));
-		Match(ShaderSource, ' ');
-		while (*ShaderSource && *ShaderSource != '\n')
-		{
-			FString ParameterName = ParseIdentifier(ShaderSource);
-			verify(ParameterName.Len() > 0);
-			verify(Match(ShaderSource, '('));
-			ParseNumber(ShaderSource);
-			verify(Match(ShaderSource, ','));
-			ParseNumber(ShaderSource);
-			verify(Match(ShaderSource, ')'));
-
-			if (Match(ShaderSource, '\n'))
-			{
-				break;
-			}
-
-			verify(Match(ShaderSource, ','));
-		}
+		check(PackedUB.Attribute.Index == Header.Bindings.NumUniformBuffers);
+		UsedUniformBufferSlots[PackedUB.Attribute.Index] = true;
+		ParameterMap.AddParameterAllocation(*PackedUB.Attribute.Name, Header.Bindings.NumUniformBuffers++, 0, 0);
+
+		// Nothing else...
+		//for (auto& Member : PackedUB.Members)
+		//{
+		//}
 	}
 
 	// Packed Uniform Buffers copy lists & setup sizes for each UB/Precision entry
@@ -911,87 +677,44 @@ static void BuildShaderOutput(
 		FlattenedUBs,
 	};
 	EFlattenUBState UBState = Unknown;
-	if (FCStringAnsi::Strncmp(ShaderSource, PackedUBCopiesPrefix, PackedUBCopiesPrefixLen) == 0)
+	for (auto& PackedUBCopy : CCHeader.PackedUBCopies)
 	{
-		ShaderSource += PackedUBCopiesPrefixLen;
-		while (*ShaderSource && *ShaderSource != '\n')
-		{
-			FOpenGLUniformBufferCopyInfo CopyInfo;
-
-			CopyInfo.SourceUBIndex = ParseNumber(ShaderSource);
-			verify(Match(ShaderSource, ':'));
-
-			CopyInfo.SourceOffsetInFloats = ParseNumber(ShaderSource);
-			verify(Match(ShaderSource, '-'));
-
-			CopyInfo.DestUBIndex = ParseNumber(ShaderSource);
-			verify(Match(ShaderSource, ':'));
-
-			CopyInfo.DestUBTypeName = *ShaderSource++;
-			CopyInfo.DestUBTypeIndex = CrossCompiler::PackedTypeNameToTypeIndex(CopyInfo.DestUBTypeName);
-			verify(Match(ShaderSource, ':'));
-
-			CopyInfo.DestOffsetInFloats = ParseNumber(ShaderSource);
-			verify(Match(ShaderSource, ':'));
-
-			CopyInfo.SizeInFloats = ParseNumber(ShaderSource);
-
-			Header.UniformBuffersCopyInfo.Add(CopyInfo);
-
-			auto& UniformBufferSize = PackedUniformBuffersSize.FindOrAdd(CopyInfo.DestUBIndex);
-			uint16& Size = UniformBufferSize.FindOrAdd(CopyInfo.DestUBTypeName);
-			Size = FMath::Max<uint16>(BytesPerComponent * (CopyInfo.DestOffsetInFloats + CopyInfo.SizeInFloats), Size);
-
-			if (Match(ShaderSource, '\n'))
-			{
-				break;
-			}
-
-			verify(Match(ShaderSource, ','));
-		}
-
-		check(UBState == Unknown);
+		CrossCompiler::FUniformBufferCopyInfo CopyInfo;
+		CopyInfo.SourceUBIndex = PackedUBCopy.SourceUB;
+		CopyInfo.SourceOffsetInFloats = PackedUBCopy.SourceOffset;
+		CopyInfo.DestUBIndex = PackedUBCopy.DestUB;
+		CopyInfo.DestUBTypeName = PackedUBCopy.DestPackedType;
+		CopyInfo.DestUBTypeIndex = CrossCompiler::PackedTypeNameToTypeIndex(CopyInfo.DestUBTypeName);
+		CopyInfo.DestOffsetInFloats = PackedUBCopy.DestOffset;
+		CopyInfo.SizeInFloats = PackedUBCopy.Count;
+
+		Header.UniformBuffersCopyInfo.Add(CopyInfo);
+
+		auto& UniformBufferSize = PackedUniformBuffersSize.FindOrAdd(CopyInfo.DestUBIndex);
+		uint16& Size = UniformBufferSize.FindOrAdd(CopyInfo.DestUBTypeName);
+		Size = FMath::Max<uint16>(BytesPerComponent * (CopyInfo.DestOffsetInFloats + CopyInfo.SizeInFloats), Size);
+
+		check(UBState == Unknown || UBState == GroupedUBs);
 		UBState = GroupedUBs;
 	}
 
-	if (FCStringAnsi::Strncmp(ShaderSource, PackedUBGlobalCopiesPrefix, PackedUBGlobalCopiesPrefixLen) == 0)
+	for (auto& PackedUBCopy : CCHeader.PackedUBGlobalCopies)
 	{
-		ShaderSource += PackedUBGlobalCopiesPrefixLen;
-		while (*ShaderSource && *ShaderSource != '\n')
-		{
-			FOpenGLUniformBufferCopyInfo CopyInfo;
-
-			CopyInfo.SourceUBIndex = ParseNumber(ShaderSource);
-			verify(Match(ShaderSource, ':'));
-
-			CopyInfo.SourceOffsetInFloats = ParseNumber(ShaderSource);
-			verify(Match(ShaderSource, '-'));
-
-			CopyInfo.DestUBIndex = 0;
-
-			CopyInfo.DestUBTypeName = *ShaderSource++;
-			CopyInfo.DestUBTypeIndex = CrossCompiler::PackedTypeNameToTypeIndex(CopyInfo.DestUBTypeName);
-			verify(Match(ShaderSource, ':'));
-
-			CopyInfo.DestOffsetInFloats = ParseNumber(ShaderSource);
-			verify(Match(ShaderSource, ':'));
-
-			CopyInfo.SizeInFloats = ParseNumber(ShaderSource);
+		CrossCompiler::FUniformBufferCopyInfo CopyInfo;
+		CopyInfo.SourceUBIndex = PackedUBCopy.SourceUB;
+		CopyInfo.SourceOffsetInFloats = PackedUBCopy.SourceOffset;
+		CopyInfo.DestUBIndex = PackedUBCopy.DestUB;
+		CopyInfo.DestUBTypeName = PackedUBCopy.DestPackedType;
+		CopyInfo.DestUBTypeIndex = CrossCompiler::PackedTypeNameToTypeIndex(CopyInfo.DestUBTypeName);
+		CopyInfo.DestOffsetInFloats = PackedUBCopy.DestOffset;
+		CopyInfo.SizeInFloats = PackedUBCopy.Count;
 
-			Header.UniformBuffersCopyInfo.Add(CopyInfo);
+		Header.UniformBuffersCopyInfo.Add(CopyInfo);
 
-			uint16& Size = PackedGlobalArraySize.FindOrAdd(CopyInfo.DestUBTypeName);
-			Size = FMath::Max<uint16>(BytesPerComponent * (CopyInfo.DestOffsetInFloats + CopyInfo.SizeInFloats), Size);
+		uint16& Size = PackedGlobalArraySize.FindOrAdd(CopyInfo.DestUBTypeName);
+		Size = FMath::Max<uint16>(BytesPerComponent * (CopyInfo.DestOffsetInFloats + CopyInfo.SizeInFloats), Size);
 
-			if (Match(ShaderSource, '\n'))
-			{
-				break;
-			}
-
-			verify(Match(ShaderSource, ','));
-		}
-
-		check(UBState == Unknown);
+		check(UBState == Unknown || UBState == FlattenedUBs);
 		UBState = FlattenedUBs;
 	}
 
@@ -1035,104 +758,49 @@ static void BuildShaderOutput(
 	}
 
 	// Then samplers.
-	if (FCStringAnsi::Strncmp(ShaderSource, SamplersPrefix, SamplersPrefixLen) == 0)
+	for (auto& Sampler : CCHeader.Samplers)
 	{
-		ShaderSource += SamplersPrefixLen;
-
-		while (*ShaderSource && *ShaderSource != '\n')
+		ParameterMap.AddParameterAllocation(
+			*Sampler.Name,
+			0,
+			Sampler.Offset,
+			Sampler.Count
+			);
+
+		Header.Bindings.NumSamplers = FMath::Max<uint8>(
+			Header.Bindings.NumSamplers,
+			Sampler.Offset + Sampler.Count
+			);
+
+		for (auto& SamplerState : Sampler.SamplerStates)
 		{
-			uint16 Offset = 0;
-			uint16 NumSamplers = 0;
-
-			FString ParameterName = ParseIdentifier(ShaderSource);
-			verify(ParameterName.Len() > 0);
-			verify(Match(ShaderSource, '('));
-			Offset = ParseNumber(ShaderSource);
-			verify(Match(ShaderSource, ':'));
-			NumSamplers = ParseNumber(ShaderSource);
 			ParameterMap.AddParameterAllocation(
-				*ParameterName,
+				*SamplerState,
 				0,
-				Offset,
-				NumSamplers
-				);
-
-			Header.Bindings.NumSamplers = FMath::Max<uint8>(
-				Header.Bindings.NumSamplers,
-				Offset + NumSamplers
+				Sampler.Offset,
+				Sampler.Count
 				);
-
-			if (Match(ShaderSource, '['))
-			{
-				// Sampler States
-				do
-				{
-					FString SamplerState = ParseIdentifier(ShaderSource);
-					checkSlow(SamplerState.Len() != 0);
-					ParameterMap.AddParameterAllocation(
-						*SamplerState,
-						0,
-						Offset,
-						NumSamplers
-						);
-				}
-				while (Match(ShaderSource, ','));
-				verify(Match(ShaderSource, ']'));
-			}
-
-			verify(Match(ShaderSource, ')'));
-
-			if (Match(ShaderSource, '\n'))
-			{
-				break;
-			}
-
-			// Skip the comma.
-			verify(Match(ShaderSource, ','));
 		}
 	}	
 
 	// Then UAVs (images in GLSL)
-	if (FCStringAnsi::Strncmp(ShaderSource, UAVsPrefix, UAVsPrefixLen) == 0)
+	for (auto& UAV : CCHeader.UAVs)
 	{
-		ShaderSource += UAVsPrefixLen;
-
-		while (*ShaderSource && *ShaderSource != '\n')
-		{
-			uint16 Offset = 0;
-			uint16 NumUAVs = 0;
-
-			FString ParameterName = ParseIdentifier(ShaderSource);
-			verify(ParameterName.Len() > 0);
-			verify(Match(ShaderSource, '('));
-			Offset = ParseNumber(ShaderSource);
-			verify(Match(ShaderSource, ':'));
-			NumUAVs = ParseNumber(ShaderSource);
-
-			ParameterMap.AddParameterAllocation(
-				*ParameterName,
-				0,
-				Offset,
-				NumUAVs
-				);
-
-			Header.Bindings.NumUAVs = FMath::Max<uint8>(
-				Header.Bindings.NumUAVs,
-				Offset + NumUAVs
-				);
-
-			verify(Match(ShaderSource, ')'));
-
-			if (Match(ShaderSource, '\n'))
-			{
-				break;
-			}
-
-			// Skip the comma.
-			verify(Match(ShaderSource, ','));
-		}
+		ParameterMap.AddParameterAllocation(
+			*UAV.Name,
+			0,
+			UAV.Offset,
+			UAV.Count
+			);
+
+		Header.Bindings.NumUAVs = FMath::Max<uint8>(
+			Header.Bindings.NumSamplers,
+			UAV.Offset + UAV.Count
+			);
 	}
 
+	Header.ShaderName = CCHeader.Name.GetCharArray();
+
 	// Build the SRT for this shader.
 	{
 		// Build the generic SRT for this shader.
@@ -1166,7 +834,7 @@ static void BuildShaderOutput(
 		// Write out the header and shader source code.
 		FMemoryWriter Ar(ShaderOutput.Code, true);
 		Ar << Header;
-		Ar.Serialize((void*)ShaderSource, SourceLen + 1 - (ShaderSource - InShaderSource));
+		Ar.Serialize((void*)USFSource, SourceLen + 1 - (USFSource - InShaderSource));
 		
 		ShaderOutput.NumInstructions = 0;
 		ShaderOutput.NumTextureSamplers = Header.Bindings.NumSamplers;
@@ -1364,7 +1032,7 @@ static void PrecompileShader(FShaderCompilerOutput& ShaderOutput, const FShaderC
 	{
 		ShaderOutput.bSucceeded = false;
 		FShaderCompilerError* NewError = new(ShaderOutput.Errors) FShaderCompilerError();
-		NewError->StrippedErrorMessage = FString::Printf(TEXT("%s shaders not supported for use in OpenGL."), GLFrequencyStringTable[ShaderInput.Target.Frequency]);
+		NewError->StrippedErrorMessage = FString::Printf(TEXT("%s shaders not supported for use in OpenGL."), CrossCompiler::GetFrequencyName((EShaderFrequency)ShaderInput.Target.Frequency));
 		return;
 	}
 
@@ -1450,34 +1118,6 @@ static void PrecompileShader(FShaderCompilerOutput& ShaderOutput, const FShaderC
 	}
 }
 
-/**
- * Parse an error emitted by the HLSL cross-compiler.
- * @param OutErrors - Array into which compiler errors may be added.
- * @param InLine - A line from the compile log.
- */
-static void ParseHlslccError(TArray<FShaderCompilerError>& OutErrors, const FString& InLine)
-{
-	const TCHAR* p = *InLine;
-	FShaderCompilerError* Error = new(OutErrors) FShaderCompilerError();
-
-	// Copy the filename.
-	while (*p && *p != TEXT('(')) { Error->ErrorFile += (*p++); }
-	Error->ErrorFile = GetRelativeShaderFilename(Error->ErrorFile);
-	p++;
-
-	// Parse the line number.
-	int32 LineNumber = 0;
-	while (*p && *p >= TEXT('0') && *p <= TEXT('9'))
-	{
-		LineNumber = 10 * LineNumber + (*p++ - TEXT('0'));
-	}
-	Error->ErrorLineString = *FString::Printf(TEXT("%d"), LineNumber);
-
-	// Skip to the warning message.
-	while (*p && (*p == TEXT(')') || *p == TEXT(':') || *p == TEXT(' ') || *p == TEXT('\t'))) { p++; }
-	Error->StrippedErrorMessage = p;
-}
-
 /*------------------------------------------------------------------------------
 	External interface.
 ------------------------------------------------------------------------------*/
@@ -1554,7 +1194,7 @@ static FString CreateCrossCompilerBatchFile( const FString& ShaderFile, const FS
 	}
 
 	const TCHAR* ApplyCSE = (CCFlags & HLSLCC_ApplyCommonSubexpressionElimination) != 0 ? TEXT("-cse") : TEXT("");
-	return CreateCrossCompilerBatchFileContents(ShaderFile, OutputFile, FrequencySwitch, EntryPoint, VersionSwitch, ApplyCSE);
+	return CrossCompiler::CreateBatchFileContents(ShaderFile, OutputFile, FrequencySwitch, EntryPoint, VersionSwitch, ApplyCSE);
 }
 
 /**
@@ -1664,7 +1304,7 @@ void CompileShader_Windows_OGL(const FShaderCompilerInput& Input,FShaderCompiler
 			FShaderCompilerError* NewError = new(Output.Errors) FShaderCompilerError();
 			NewError->StrippedErrorMessage = FString::Printf(
 				TEXT("%s shaders not supported for use in OpenGL."),
-				GLFrequencyStringTable[Input.Target.Frequency]
+				CrossCompiler::GetFrequencyName((EShaderFrequency)Input.Target.Frequency)
 				);
 			return;
 		}
@@ -1793,7 +1433,7 @@ void CompileShader_Windows_OGL(const FShaderCompilerInput& Input,FShaderCompiler
 			for (int32 LineIndex = 0; LineIndex < ErrorLines.Num(); ++LineIndex)
 			{
 				const FString& Line = ErrorLines[LineIndex];
-				ParseHlslccError(Output.Errors, Line);
+				CrossCompiler::ParseHlslccError(Output.Errors, Line);
 			}
 		}
 
diff --git a/Engine/Source/Developer/ShaderFormatOpenGL/Private/ShaderFormatOpenGL.cpp b/Engine/Source/Developer/ShaderFormatOpenGL/Private/ShaderFormatOpenGL.cpp
--- a/Engine/Source/Developer/ShaderFormatOpenGL/Private/ShaderFormatOpenGL.cpp
+++ b/Engine/Source/Developer/ShaderFormatOpenGL/Private/ShaderFormatOpenGL.cpp
@@ -25,16 +25,16 @@ class FShaderFormatGLSL : public IShaderFormat
 	enum 
 	{
 		/** Version for shader format, this becomes part of the DDC key. */
-		UE_SHADER_GLSL_150_VER = 26,
-		UE_SHADER_GLSL_150_MAC_VER = 6,
-		UE_SHADER_GLSL_430_VER = 58,
-		UE_SHADER_GLSL_ES2_VER = 15,
-		UE_SHADER_GLSL_150ES2_VER = 18,
-		UE_SHADER_GLSL_150ES2NOUB_VER = 2,
-		UE_SHADER_GLSL_150ES3_1_VER = 2,
-		UE_SHADER_GLSL_ES2_VER_WEBGL  = 16,
-		UE_SHADER_GLSL_ES2_IOS_VER  = 5,
-		UE_SHADER_GLSL_310_ES_EXT_VER = 4,
+		UE_SHADER_GLSL_150_VER = 27,
+		UE_SHADER_GLSL_150_MAC_VER = 7,
+		UE_SHADER_GLSL_430_VER = 59,
+		UE_SHADER_GLSL_ES2_VER = 16,
+		UE_SHADER_GLSL_150ES2_VER = 19,
+		UE_SHADER_GLSL_150ES2NOUB_VER = 3,
+		UE_SHADER_GLSL_150ES3_1_VER = 3,
+		UE_SHADER_GLSL_ES2_VER_WEBGL  = 17,
+		UE_SHADER_GLSL_ES2_IOS_VER  = 6,
+		UE_SHADER_GLSL_310_ES_EXT_VER = 5,
 	}; 
 
 	void CheckFormat(FName Format) const
diff --git a/Engine/Source/Developer/iOS/MetalShaderFormat/Private/MetalShaderCompiler.cpp b/Engine/Source/Developer/iOS/MetalShaderFormat/Private/MetalShaderCompiler.cpp
--- a/Engine/Source/Developer/iOS/MetalShaderFormat/Private/MetalShaderCompiler.cpp
+++ b/Engine/Source/Developer/iOS/MetalShaderFormat/Private/MetalShaderCompiler.cpp
@@ -27,45 +27,7 @@ DEFINE_LOG_CATEGORY_STATIC(LogMetalShaderCompiler, Log, All);
 	Shader compiling.
 ------------------------------------------------------------------------------*/
 
-/** Map shader frequency -> string for messages. */
-static const TCHAR* GLFrequencyStringTable[] =
-{
-	TEXT("Vertex"),
-	TEXT("Hull"),
-	TEXT("Domain"),
-	TEXT("Pixel"),
-	TEXT("Geometry"),
-	TEXT("Compute")
-};
-
-static FString ParseIdentifier(const ANSICHAR* &Str)
-{
-	FString Result;
-
-	while ((*Str >= 'A' && *Str <= 'Z')
-		|| (*Str >= 'a' && *Str <= 'z')
-		|| (*Str >= '0' && *Str <= '9')
-		|| *Str == '_')
-	{
-		Result += *Str;
-		++Str;
-	}
-
-	return Result;
-}
-
-static bool Match(const ANSICHAR* &Str, ANSICHAR Char)
-{
-	if (*Str == Char)
-	{
-		++Str;
-		return true;
-	}
-
-	return false;
-}
-
-static uint32 ParseNumber(const ANSICHAR* &Str)
+static inline uint32 ParseNumber(const TCHAR* Str)
 {
 	uint32 Num = 0;
 	while (*Str && *Str >= '0' && *Str <= '9')
@@ -90,8 +52,14 @@ static void BuildMetalShaderOutput(
 	TArray<FShaderCompilerError>& OutErrors
 	)
 {
+	const ANSICHAR* USFSource = InShaderSource;
+	CrossCompiler::FHlslccHeader CCHeader;
+	if (!CCHeader.Read(USFSource, SourceLen))
+	{
+		UE_LOG(LogMetalShaderCompiler, Fatal, TEXT("Bad hlslcc header found"));
+	}
+
 	FMetalCodeHeader Header = {0};
-	const ANSICHAR* ShaderSource = InShaderSource;
 	FShaderParameterMap& ParameterMap = ShaderOutput.ParameterMap;
 	EShaderFrequency Frequency = (EShaderFrequency)ShaderOutput.Target.Frequency;
 
@@ -101,378 +69,122 @@ static void BuildMetalShaderOutput(
 	// Write out the magic markers.
 	Header.Frequency = Frequency;
 
-	#define DEF_PREFIX_STR(Str) \
-		const ANSICHAR* Str##Prefix = "// @" #Str ": "; \
-		const int32 Str##PrefixLen = FCStringAnsi::Strlen(Str##Prefix)
-	DEF_PREFIX_STR(Inputs);
-	DEF_PREFIX_STR(Outputs);
-	DEF_PREFIX_STR(UniformBlocks);
-	DEF_PREFIX_STR(Uniforms);
-	DEF_PREFIX_STR(PackedGlobals);
-	DEF_PREFIX_STR(PackedUB);
-	DEF_PREFIX_STR(PackedUBCopies);
-	DEF_PREFIX_STR(PackedUBGlobalCopies);
-	DEF_PREFIX_STR(Samplers);
-	DEF_PREFIX_STR(UAVs);
-	DEF_PREFIX_STR(SamplerStates);
-	DEF_PREFIX_STR(NumThreads);
-	#undef DEF_PREFIX_STR
-
-	ANSICHAR const* ShaderName = nullptr;
-	uint32 ShaderNameLen = 0;
-	while (	FCStringAnsi::Strncmp(ShaderSource, "// !", 4) == 0 )
+	// Only inputs for vertex shaders must be tracked.
+	if (Frequency == SF_Vertex)
 	{
-		ShaderName = ShaderSource;
-		while (*ShaderSource && *ShaderSource++ != '\n') { ShaderNameLen++; }
-	}
-	
-	// Skip any comments that come before the signature.
-	while (	FCStringAnsi::Strncmp(ShaderSource, "//", 2) == 0 &&
-			FCStringAnsi::Strncmp(ShaderSource, "// @", 4) != 0 )
-	{
-		while (*ShaderSource && *ShaderSource++ != '\n') {}
-	}
-
-	// HLSLCC first prints the list of inputs.
-	if (FCStringAnsi::Strncmp(ShaderSource, InputsPrefix, InputsPrefixLen) == 0)
-	{
-		ShaderSource += InputsPrefixLen;
-
-		// Only inputs for vertex shaders must be tracked.
-		if (Frequency == SF_Vertex)
+		static const FString AttributePrefix = TEXT("in_ATTRIBUTE");
+		for (auto& Input : CCHeader.Inputs)
 		{
-			const ANSICHAR* AttributePrefix = "in_ATTRIBUTE";
-			const int32 AttributePrefixLen = FCStringAnsi::Strlen(AttributePrefix);
-			while (*ShaderSource && *ShaderSource != '\n')
+			// Only process attributes.
+			if (Input.Name.StartsWith(AttributePrefix))
 			{
-				// Skip the type.
-				while (*ShaderSource && *ShaderSource++ != ':') {}
-				
-				// Only process attributes.
-				if (FCStringAnsi::Strncmp(ShaderSource, AttributePrefix, AttributePrefixLen) == 0)
-				{
-					ShaderSource += AttributePrefixLen;
-					uint8 AttributeIndex = ParseNumber(ShaderSource);
-					Header.Bindings.InOutMask |= (1 << AttributeIndex);
-				}
-
-				// Skip to the next.
-				while (*ShaderSource && *ShaderSource != ',' && *ShaderSource != '\n')
-				{
-					ShaderSource++;
-				}
-
-				if (Match(ShaderSource, '\n'))
-				{
-					break;
-				}
-
-				verify(Match(ShaderSource, ','));
+				uint8 AttributeIndex = ParseNumber(*Input.Name + AttributePrefix.Len());
+				Header.Bindings.InOutMask |= (1 << AttributeIndex);
 			}
 		}
-		else
-		{
-			// Skip to the next line.
-			while (*ShaderSource && *ShaderSource++ != '\n') {}
-		}
 	}
 
 	// Then the list of outputs.
-	if (FCStringAnsi::Strncmp(ShaderSource, OutputsPrefix, OutputsPrefixLen) == 0)
+	static const FString TargetPrefix = "out_Target";
+	static const FString GL_FragDepth = "gl_FragDepth";
+	// Only outputs for pixel shaders must be tracked.
+	if (Frequency == SF_Pixel)
 	{
-		ShaderSource += OutputsPrefixLen;
-
-		// Only outputs for pixel shaders must be tracked.
-		if (Frequency == SF_Pixel)
+		for (auto& Output : CCHeader.Outputs)
 		{
-			const ANSICHAR* TargetPrefix = "out_Target";
-			const int32 TargetPrefixLen = FCStringAnsi::Strlen(TargetPrefix);
-
-			while (*ShaderSource && *ShaderSource != '\n')
+			// Handle targets.
+			if (Output.Name.StartsWith(TargetPrefix))
 			{
-				// Skip the type.
-				while (*ShaderSource && *ShaderSource++ != ':') {}
-
-				// Handle targets.
-				if (FCStringAnsi::Strncmp(ShaderSource, TargetPrefix, TargetPrefixLen) == 0)
-				{
-					ShaderSource += TargetPrefixLen;
-					uint8 TargetIndex = ParseNumber(ShaderSource);
-					Header.Bindings.InOutMask |= (1 << TargetIndex);
-				}
-				// Handle depth writes.
-				else if (FCStringAnsi::Strcmp(ShaderSource, "gl_FragDepth") == 0)
-				{
-					Header.Bindings.InOutMask |= 0x8000;
-				}
-
-				// Skip to the next.
-				while (*ShaderSource && *ShaderSource != ',' && *ShaderSource != '\n')
-				{
-					ShaderSource++;
-				}
-
-				if (Match(ShaderSource, '\n'))
-				{
-					break;
-				}
-
-				verify(Match(ShaderSource, ','));
+				uint8 TargetIndex = ParseNumber(*Output.Name + TargetPrefix.Len());
+				Header.Bindings.InOutMask |= (1 << TargetIndex);
+			}
+			// Handle depth writes.
+			else if (Output.Name.Equals(GL_FragDepth))
+			{
+				Header.Bindings.InOutMask |= 0x8000;
 			}
-		}
-		else
-		{
-			// Skip to the next line.
-			while (*ShaderSource && *ShaderSource++ != '\n') {}
 		}
 	}
 
 	bool bHasRegularUniformBuffers = false;
 
 	// Then 'normal' uniform buffers.
-	if (FCStringAnsi::Strncmp(ShaderSource, UniformBlocksPrefix, UniformBlocksPrefixLen) == 0)
-	{
-		ShaderSource += UniformBlocksPrefixLen;
-
-		while (*ShaderSource && *ShaderSource != '\n')
-		{
-			FString BufferName = ParseIdentifier(ShaderSource);
-			verify(BufferName.Len() > 0);
-			verify(Match(ShaderSource, '('));
-			uint16 UBIndex = ParseNumber(ShaderSource);
-			if (UBIndex >= Header.Bindings.NumUniformBuffers)
-			{
-				Header.Bindings.NumUniformBuffers = UBIndex + 1;
-			}
-			UsedUniformBufferSlots[UBIndex] = true;
-			verify(Match(ShaderSource, ')'));
-			ParameterMap.AddParameterAllocation(*BufferName, UBIndex, 0, 0);
-			bHasRegularUniformBuffers = true;
-
-			// Skip the comma.
-			if (Match(ShaderSource, '\n'))
-			{
-				break;
-			}
-
-			verify(Match(ShaderSource, ','));
-		}
-
-		Match(ShaderSource, '\n');
-	}
-
-	// Then uniforms.
-	const uint16 BytesPerComponent = 4;
-/*
-	uint16 PackedUniformSize[OGL_NUM_PACKED_UNIFORM_ARRAYS] = {0};
-	FMemory::Memzero(&PackedUniformSize, sizeof(PackedUniformSize));
-*/
-	if (FCStringAnsi::Strncmp(ShaderSource, UniformsPrefix, UniformsPrefixLen) == 0)
+	for (auto& UniformBlock : CCHeader.UniformBlocks)
 	{
-		// @todo-mobile: Will we ever need to support this code path?
-		check(0);
-/*
-		ShaderSource += UniformsPrefixLen;
-
-		while (*ShaderSource && *ShaderSource != '\n')
+		uint16 UBIndex = UniformBlock.Index;
+		if (UBIndex >= Header.Bindings.NumUniformBuffers)
 		{
-			uint16 ArrayIndex = 0;
-			uint16 Offset = 0;
-			uint16 NumComponents = 0;
-
-			FString ParameterName = ParseIdentifier(ShaderSource);
-			verify(ParameterName.Len() > 0);
-			verify(Match(ShaderSource, '('));
-			ArrayIndex = ParseNumber(ShaderSource);
-			verify(Match(ShaderSource, ':'));
-			Offset = ParseNumber(ShaderSource);
-			verify(Match(ShaderSource, ':'));
-			NumComponents = ParseNumber(ShaderSource);
-			verify(Match(ShaderSource, ')'));
-
-			ParameterMap.AddParameterAllocation(
-				*ParameterName,
-				ArrayIndex,
-				Offset * BytesPerComponent,
-				NumComponents * BytesPerComponent
-				);
-
-			if (ArrayIndex < OGL_NUM_PACKED_UNIFORM_ARRAYS)
-			{
-				PackedUniformSize[ArrayIndex] = FMath::Max<uint16>(
-					PackedUniformSize[ArrayIndex],
-					BytesPerComponent * (Offset + NumComponents)
-					);
-			}
-
-			// Skip the comma.
-			if (Match(ShaderSource, '\n'))
-			{
-				break;
-			}
-
-			verify(Match(ShaderSource, ','));
+			Header.Bindings.NumUniformBuffers = UBIndex + 1;
 		}
-
-		Match(ShaderSource, '\n');
-*/
+		UsedUniformBufferSlots[UBIndex] = true;
+		ParameterMap.AddParameterAllocation(*UniformBlock.Name, UBIndex, 0, 0);
+		bHasRegularUniformBuffers = true;
 	}
 
 	// Packed global uniforms
+	const uint16 BytesPerComponent = 4;
 	TMap<ANSICHAR, uint16> PackedGlobalArraySize;
-	if (FCStringAnsi::Strncmp(ShaderSource, PackedGlobalsPrefix, PackedGlobalsPrefixLen) == 0)
+	for (auto& PackedGlobal : CCHeader.PackedGlobals)
 	{
-		ShaderSource += PackedGlobalsPrefixLen;
-		while (*ShaderSource && *ShaderSource != '\n')
-		{
-			ANSICHAR ArrayIndex = 0;
-			uint16 Offset = 0;
-			uint16 NumComponents = 0;
-
-			FString ParameterName = ParseIdentifier(ShaderSource);
-			verify(ParameterName.Len() > 0);
-			verify(Match(ShaderSource, '('));
-			ArrayIndex = *ShaderSource++;
-			verify(Match(ShaderSource, ':'));
-			Offset = ParseNumber(ShaderSource);
-			verify(Match(ShaderSource, ','));
-			NumComponents = ParseNumber(ShaderSource);
-			verify(Match(ShaderSource, ')'));
-
-			ParameterMap.AddParameterAllocation(
-				*ParameterName,
-				ArrayIndex,
-				Offset * BytesPerComponent,
-				NumComponents * BytesPerComponent
-				);
-
-			uint16& Size = PackedGlobalArraySize.FindOrAdd(ArrayIndex);
-			Size = FMath::Max<uint16>(BytesPerComponent * (Offset + NumComponents), Size);
-
-			if (Match(ShaderSource, '\n'))
-			{
-				break;
-			}
+		ParameterMap.AddParameterAllocation(
+			*PackedGlobal.Name,
+			PackedGlobal.PackedType,
+			PackedGlobal.Offset * BytesPerComponent,
+			PackedGlobal.Count * BytesPerComponent
+			);
 
-			// Skip the comma.
-			verify(Match(ShaderSource, ','));
-		}
-
-		Match(ShaderSource, '\n');
+		uint16& Size = PackedGlobalArraySize.FindOrAdd(PackedGlobal.PackedType);
+		Size = FMath::Max<uint16>(BytesPerComponent * (PackedGlobal.Offset + PackedGlobal.Count), Size);
 	}
 
 	// Packed Uniform Buffers
 	TMap<int, TMap<ANSICHAR, uint16> > PackedUniformBuffersSize;
-	while (FCStringAnsi::Strncmp(ShaderSource, PackedUBPrefix, PackedUBPrefixLen) == 0)
+	for (auto& PackedUB : CCHeader.PackedUBs)
 	{
-		ShaderSource += PackedUBPrefixLen;
-		FString BufferName = ParseIdentifier(ShaderSource);
-		verify(BufferName.Len() > 0);
-		verify(Match(ShaderSource, '('));
-		uint16 BufferIndex = ParseNumber(ShaderSource);
-		check(BufferIndex == Header.Bindings.NumUniformBuffers);
-		UsedUniformBufferSlots[BufferIndex] = true;
-		verify(Match(ShaderSource, ')'));
-		ParameterMap.AddParameterAllocation(*BufferName, Header.Bindings.NumUniformBuffers++, 0, 0);
-
-		verify(Match(ShaderSource, ':'));
-		Match(ShaderSource, ' ');
-		while (*ShaderSource && *ShaderSource != '\n')
-		{
-			FString ParameterName = ParseIdentifier(ShaderSource);
-			verify(ParameterName.Len() > 0);
-			verify(Match(ShaderSource, '('));
-			ParseNumber(ShaderSource);
-			verify(Match(ShaderSource, ','));
-			ParseNumber(ShaderSource);
-			verify(Match(ShaderSource, ')'));
-
-			if (Match(ShaderSource, '\n'))
-			{
-				break;
-			}
+		check(PackedUB.Attribute.Index == Header.Bindings.NumUniformBuffers);
+		UsedUniformBufferSlots[PackedUB.Attribute.Index] = true;
+		ParameterMap.AddParameterAllocation(*PackedUB.Attribute.Name, Header.Bindings.NumUniformBuffers++, 0, 0);
 
-			verify(Match(ShaderSource, ','));
-		}
+		// Nothing else...
+		//for (auto& Member : PackedUB.Members)
+		//{
+		//}
 	}
 
 	// Packed Uniform Buffers copy lists & setup sizes for each UB/Precision entry
-	if (FCStringAnsi::Strncmp(ShaderSource, PackedUBCopiesPrefix, PackedUBCopiesPrefixLen) == 0)
+	for (auto& PackedUBCopy : CCHeader.PackedUBCopies)
 	{
-		ShaderSource += PackedUBCopiesPrefixLen;
-		while (*ShaderSource && *ShaderSource != '\n')
-		{
-			FMetalUniformBufferCopyInfo CopyInfo;
-
-			CopyInfo.SourceUBIndex = ParseNumber(ShaderSource);
-			verify(Match(ShaderSource, ':'));
-
-			CopyInfo.SourceOffsetInFloats = ParseNumber(ShaderSource);
-			verify(Match(ShaderSource, '-'));
-
-			CopyInfo.DestUBIndex = ParseNumber(ShaderSource);
-			verify(Match(ShaderSource, ':'));
-
-			CopyInfo.DestUBTypeName = *ShaderSource++;
-			CopyInfo.DestUBTypeIndex = CrossCompiler::PackedTypeNameToTypeIndex(CopyInfo.DestUBTypeName);
-			verify(Match(ShaderSource, ':'));
-
-			CopyInfo.DestOffsetInFloats = ParseNumber(ShaderSource);
-			verify(Match(ShaderSource, ':'));
-
-			CopyInfo.SizeInFloats = ParseNumber(ShaderSource);
-
-			Header.UniformBuffersCopyInfo.Add(CopyInfo);
-
-			auto& UniformBufferSize = PackedUniformBuffersSize.FindOrAdd(CopyInfo.DestUBIndex);
-			uint16& Size = UniformBufferSize.FindOrAdd(CopyInfo.DestUBTypeName);
-			Size = FMath::Max<uint16>(BytesPerComponent * (CopyInfo.DestOffsetInFloats + CopyInfo.SizeInFloats), Size);
+		CrossCompiler::FUniformBufferCopyInfo CopyInfo;
+		CopyInfo.SourceUBIndex = PackedUBCopy.SourceUB;
+		CopyInfo.SourceOffsetInFloats = PackedUBCopy.SourceOffset;
+		CopyInfo.DestUBIndex = PackedUBCopy.DestUB;
+		CopyInfo.DestUBTypeName = PackedUBCopy.DestPackedType;
+		CopyInfo.DestUBTypeIndex = CrossCompiler::PackedTypeNameToTypeIndex(CopyInfo.DestUBTypeName);
+		CopyInfo.DestOffsetInFloats = PackedUBCopy.DestOffset;
+		CopyInfo.SizeInFloats = PackedUBCopy.Count;
 
-			if (Match(ShaderSource, '\n'))
-			{
-				break;
-			}
+		Header.UniformBuffersCopyInfo.Add(CopyInfo);
 
-			verify(Match(ShaderSource, ','));
-		}
+		auto& UniformBufferSize = PackedUniformBuffersSize.FindOrAdd(CopyInfo.DestUBIndex);
+		uint16& Size = UniformBufferSize.FindOrAdd(CopyInfo.DestUBTypeName);
+		Size = FMath::Max<uint16>(BytesPerComponent * (CopyInfo.DestOffsetInFloats + CopyInfo.SizeInFloats), Size);
 	}
 
-	if (FCStringAnsi::Strncmp(ShaderSource, PackedUBGlobalCopiesPrefix, PackedUBGlobalCopiesPrefixLen) == 0)
+	for (auto& PackedUBCopy : CCHeader.PackedUBGlobalCopies)
 	{
-		ShaderSource += PackedUBGlobalCopiesPrefixLen;
-		while (*ShaderSource && *ShaderSource != '\n')
-		{
-			FMetalUniformBufferCopyInfo CopyInfo;
-
-			CopyInfo.SourceUBIndex = ParseNumber(ShaderSource);
-			verify(Match(ShaderSource, ':'));
-
-			CopyInfo.SourceOffsetInFloats = ParseNumber(ShaderSource);
-			verify(Match(ShaderSource, '-'));
-
-			CopyInfo.DestUBIndex = 0;
-
-			CopyInfo.DestUBTypeName = *ShaderSource++;
-			CopyInfo.DestUBTypeIndex = CrossCompiler::PackedTypeNameToTypeIndex(CopyInfo.DestUBTypeName);
-			verify(Match(ShaderSource, ':'));
-
-			CopyInfo.DestOffsetInFloats = ParseNumber(ShaderSource);
-			verify(Match(ShaderSource, ':'));
-
-			CopyInfo.SizeInFloats = ParseNumber(ShaderSource);
+		CrossCompiler::FUniformBufferCopyInfo CopyInfo;
+		CopyInfo.SourceUBIndex = PackedUBCopy.SourceUB;
+		CopyInfo.SourceOffsetInFloats = PackedUBCopy.SourceOffset;
+		CopyInfo.DestUBIndex = PackedUBCopy.DestUB;
+		CopyInfo.DestUBTypeName = PackedUBCopy.DestPackedType;
+		CopyInfo.DestUBTypeIndex = CrossCompiler::PackedTypeNameToTypeIndex(CopyInfo.DestUBTypeName);
+		CopyInfo.DestOffsetInFloats = PackedUBCopy.DestOffset;
+		CopyInfo.SizeInFloats = PackedUBCopy.Count;
 
-			Header.UniformBuffersCopyInfo.Add(CopyInfo);
+		Header.UniformBuffersCopyInfo.Add(CopyInfo);
 
-			uint16& Size = PackedGlobalArraySize.FindOrAdd(CopyInfo.DestUBTypeName);
-			Size = FMath::Max<uint16>(BytesPerComponent * (CopyInfo.DestOffsetInFloats + CopyInfo.SizeInFloats), Size);
-
-			if (Match(ShaderSource, '\n'))
-			{
-				break;
-			}
-
-			verify(Match(ShaderSource, ','));
-		}
+		uint16& Size = PackedGlobalArraySize.FindOrAdd(CopyInfo.DestUBTypeName);
+		Size = FMath::Max<uint16>(BytesPerComponent * (CopyInfo.DestOffsetInFloats + CopyInfo.SizeInFloats), Size);
 	}
 	Header.Bindings.bHasRegularUniformBuffers = bHasRegularUniformBuffers;
 
@@ -515,138 +227,55 @@ static void BuildMetalShaderOutput(
 
 	// Then samplers.
 	TMap<FString, uint32> SamplerMap;
-	if (FCStringAnsi::Strncmp(ShaderSource, SamplersPrefix, SamplersPrefixLen) == 0)
+	for (auto& Sampler : CCHeader.Samplers)
 	{
-		ShaderSource += SamplersPrefixLen;
-
-		while (*ShaderSource && *ShaderSource != '\n')
-		{
-			uint16 Offset = 0;
-			uint16 NumSamplers = 0;
-
-			FString ParameterName = ParseIdentifier(ShaderSource);
-			verify(ParameterName.Len() > 0);
-			verify(Match(ShaderSource, '('));
-			Offset = ParseNumber(ShaderSource);
-			verify(Match(ShaderSource, ':'));
-			NumSamplers = ParseNumber(ShaderSource);
-			ParameterMap.AddParameterAllocation(
-				*ParameterName,
-				0,
-				Offset,
-				NumSamplers
-				);
-
-			Header.Bindings.NumSamplers = FMath::Max<uint8>(
-				Header.Bindings.NumSamplers,
-				Offset + NumSamplers
-				);
+		ParameterMap.AddParameterAllocation(
+			*Sampler.Name,
+			0,
+			Sampler.Offset,
+			Sampler.Count
+			);
 
-			if (Match(ShaderSource, '['))
-			{
-				// Sampler States
-				do
-				{
-					FString SamplerState = ParseIdentifier(ShaderSource);
-					checkSlow(SamplerState.Len() != 0);
-					SamplerMap.Add(SamplerState, NumSamplers);
-				}
-				while (Match(ShaderSource, ','));
-				verify(Match(ShaderSource, ']'));
-			}
-
-			verify(Match(ShaderSource, ')'));
+		Header.Bindings.NumSamplers = FMath::Max<uint8>(
+			Header.Bindings.NumSamplers,
+			Sampler.Offset + Sampler.Count
+			);
 
-			if (Match(ShaderSource, '\n'))
-			{
-				break;
-			}
-
-			// Skip the comma.
-			verify(Match(ShaderSource, ','));
+		for (auto& SamplerState : Sampler.SamplerStates)
+		{
+			SamplerMap.Add(SamplerState, Sampler.Count);
 		}
 	}	
 
 	// Then UAVs (images in Metal)
-	if (FCStringAnsi::Strncmp(ShaderSource, UAVsPrefix, UAVsPrefixLen) == 0)
+	for (auto& UAV : CCHeader.UAVs)
 	{
-		ShaderSource += UAVsPrefixLen;
-
-		while (*ShaderSource && *ShaderSource != '\n')
-		{
-			uint16 Offset = 0;
-			uint16 NumUAVs = 0;
-
-			FString ParameterName = ParseIdentifier(ShaderSource);
-			verify(ParameterName.Len() > 0);
-			verify(Match(ShaderSource, '('));
-			Offset = ParseNumber(ShaderSource);
-			verify(Match(ShaderSource, ':'));
-			NumUAVs = ParseNumber(ShaderSource);
-
-			ParameterMap.AddParameterAllocation(
-				*ParameterName,
-				0,
-				Offset,
-				NumUAVs
-				);
-
-			Header.Bindings.NumUAVs = FMath::Max<uint8>(
-				Header.Bindings.NumUAVs,
-				Offset + NumUAVs
-				);
-
-			verify(Match(ShaderSource, ')'));
+		ParameterMap.AddParameterAllocation(
+			*UAV.Name,
+			0,
+			UAV.Offset,
+			UAV.Count
+			);
 
-			if (Match(ShaderSource, '\n'))
-			{
-				break;
-			}
-
-			// Skip the comma.
-			verify(Match(ShaderSource, ','));
-		}
+		Header.Bindings.NumUAVs = FMath::Max<uint8>(
+			Header.Bindings.NumSamplers,
+			UAV.Offset + UAV.Count
+			);
 	}
 
-	if (FCStringAnsi::Strncmp(ShaderSource, SamplerStatesPrefix, SamplerStatesPrefixLen) == 0)
+	for (auto& SamplerState : CCHeader.SamplerStates)
 	{
-		ShaderSource += SamplerStatesPrefixLen;
-		while (*ShaderSource && *ShaderSource != '\n')
-		{
-			int32 Index = ParseNumber(ShaderSource);
-			verify(Match(ShaderSource, ':'));
-			FString SamplerState = ParseIdentifier(ShaderSource);
-
-			checkSlow(SamplerState.Len() != 0);
-			ParameterMap.AddParameterAllocation(
-				*SamplerState,
-				0,
-				Index,
-				SamplerMap[SamplerState]
-				);
-			
-			if (Match(ShaderSource, '\n'))
-			{
-				break;
-			}
-
-			// Skip the comma.
-			verify(Match(ShaderSource, ','));
-		}
+		ParameterMap.AddParameterAllocation(
+			*SamplerState.Name,
+			0,
+			SamplerState.Index,
+			SamplerMap[SamplerState.Name]
+			);
 	}
 
-	if (FCStringAnsi::Strncmp(ShaderSource, NumThreadsPrefix, NumThreadsPrefixLen) == 0)
-	{
-		ShaderSource += NumThreadsPrefixLen;
-		Header.NumThreadsX = ParseNumber(ShaderSource);
-		verify(Match(ShaderSource, ','));
-		Match(ShaderSource, ' ');
-		Header.NumThreadsY = ParseNumber(ShaderSource);
-		verify(Match(ShaderSource, ','));
-		Match(ShaderSource, ' ');
-		Header.NumThreadsZ = ParseNumber(ShaderSource);
-		verify(Match(ShaderSource, '\n'));
-	}
+	Header.NumThreadsX = CCHeader.NumThreads[0];
+	Header.NumThreadsY = CCHeader.NumThreads[1];
+	Header.NumThreadsZ = CCHeader.NumThreads[2];
 
 	// Build the SRT for this shader.
 	{
@@ -669,6 +298,7 @@ static void BuildMetalShaderOutput(
 	}
 
 	const int32 MaxSamplers = GetFeatureLevelMaxTextureSamplers(ERHIFeatureLevel::ES3_1);
+	Header.ShaderName = CCHeader.Name.GetCharArray();
 
 	if (Header.Bindings.NumSamplers > MaxSamplers)
 	{
@@ -685,11 +315,7 @@ static void BuildMetalShaderOutput(
 		uint8 PrecompiledFlag = 0;
 		Ar << PrecompiledFlag;
 		Ar << Header;
-		if(ShaderName && ShaderNameLen)
-		{
-			Ar.Serialize((void*)ShaderName, ShaderNameLen);
-		}
-		Ar.Serialize((void*)ShaderSource, SourceLen + 1 - (ShaderSource - InShaderSource));
+		Ar.Serialize((void*)USFSource, SourceLen + 1 - (USFSource - InShaderSource));
 		
 		ShaderOutput.NumInstructions = 0;
 		ShaderOutput.NumTextureSamplers = Header.Bindings.NumSamplers;
@@ -705,7 +331,7 @@ static void BuildMetalShaderOutput(
 		InputFilename = InputFilename + TEXT(".metal");
 		
 		// write out shader source
-		FFileHelper::SaveStringToFile(FString(ShaderSource), *InputFilename);
+		FFileHelper::SaveStringToFile(FString(USFSource), *InputFilename);
 		
 		int32 ReturnCode = 0;
 		FString Results;
@@ -807,7 +433,7 @@ static void BuildMetalShaderOutput(
 			uint8 PrecompiledFlag = 0;
 			Ar << PrecompiledFlag;
 			Ar << Header;
-			Ar.Serialize((void*)ShaderSource, SourceLen + 1 - (ShaderSource - InShaderSource));
+			Ar.Serialize((void*)USFSource, SourceLen + 1 - (USFSource - InShaderSource));
 			
 			ShaderOutput.NumInstructions = 0;
 			ShaderOutput.NumTextureSamplers = Header.Bindings.NumSamplers;
@@ -823,34 +449,6 @@ static void BuildMetalShaderOutput(
 	}
 }
 
-/**
- * Parse an error emitted by the HLSL cross-compiler.
- * @param OutErrors - Array into which compiler errors may be added.
- * @param InLine - A line from the compile log.
- */
-static void ParseHlslccError(TArray<FShaderCompilerError>& OutErrors, const FString& InLine)
-{
-	const TCHAR* p = *InLine;
-	FShaderCompilerError* Error = new(OutErrors) FShaderCompilerError();
-
-	// Copy the filename.
-	while (*p && *p != TEXT('(')) { Error->ErrorFile += (*p++); }
-	Error->ErrorFile = GetRelativeShaderFilename(Error->ErrorFile);
-	p++;
-
-	// Parse the line number.
-	int32 LineNumber = 0;
-	while (*p && *p >= TEXT('0') && *p <= TEXT('9'))
-	{
-		LineNumber = 10 * LineNumber + (*p++ - TEXT('0'));
-	}
-	Error->ErrorLineString = *FString::Printf(TEXT("%d"), LineNumber);
-
-	// Skip to the warning message.
-	while (*p && (*p == TEXT(')') || *p == TEXT(':') || *p == TEXT(' ') || *p == TEXT('\t'))) { p++; }
-	Error->StrippedErrorMessage = p;
-}
-
 /*------------------------------------------------------------------------------
 	External interface.
 ------------------------------------------------------------------------------*/
@@ -890,7 +488,7 @@ static FString CreateCommandLineHLSLCC( const FString& ShaderFile, const FString
 	const TCHAR* VersionSwitch = TEXT("-metal");
 	const TCHAR* FlattenUB = ((CCFlags & HLSLCC_FlattenUniformBuffers) == HLSLCC_FlattenUniformBuffers) ? TEXT("-flattenub") : TEXT("");
 
-	return CreateCrossCompilerBatchFileContents(ShaderFile, OutputFile, FrequencySwitch, EntryPoint, VersionSwitch, FlattenUB);
+	return CrossCompiler::CreateBatchFileContents(ShaderFile, OutputFile, FrequencySwitch, EntryPoint, VersionSwitch, FlattenUB);
 }
 
 void CompileShader_Metal(const FShaderCompilerInput& Input,FShaderCompilerOutput& Output,const FString& WorkingDirectory)
@@ -969,7 +567,7 @@ void CompileShader_Metal(const FShaderCompilerInput& Input,FShaderCompilerOutput
 			FShaderCompilerError* NewError = new(Output.Errors) FShaderCompilerError();
 			NewError->StrippedErrorMessage = FString::Printf(
 				TEXT("%s shaders not supported for use in Metal."),
-				GLFrequencyStringTable[Input.Target.Frequency]
+				CrossCompiler::GetFrequencyName((EShaderFrequency)Input.Target.Frequency)
 				);
 			return;
 		}
@@ -1068,7 +666,7 @@ void CompileShader_Metal(const FShaderCompilerInput& Input,FShaderCompilerOutput
 			for (int32 LineIndex = 0; LineIndex < ErrorLines.Num(); ++LineIndex)
 			{
 				const FString& Line = ErrorLines[LineIndex];
-				ParseHlslccError(Output.Errors, Line);
+				CrossCompiler::ParseHlslccError(Output.Errors, Line);
 			}
 		}
 
diff --git a/Engine/Source/Developer/iOS/MetalShaderFormat/Private/MetalShaderFormat.cpp b/Engine/Source/Developer/iOS/MetalShaderFormat/Private/MetalShaderFormat.cpp
--- a/Engine/Source/Developer/iOS/MetalShaderFormat/Private/MetalShaderFormat.cpp
+++ b/Engine/Source/Developer/iOS/MetalShaderFormat/Private/MetalShaderFormat.cpp
@@ -16,7 +16,7 @@ class FMetalShaderFormat : public IShaderFormat
 {
 	enum
 	{
-		HEADER_VERSION = 17,
+		HEADER_VERSION = 18,
 	};
 	struct FVersion
 	{
diff --git a/Engine/Source/Programs/CrossCompilerTool/Private/CrossCompilerTool.cpp b/Engine/Source/Programs/CrossCompilerTool/Private/CrossCompilerTool.cpp
--- a/Engine/Source/Programs/CrossCompilerTool/Private/CrossCompilerTool.cpp
+++ b/Engine/Source/Programs/CrossCompilerTool/Private/CrossCompilerTool.cpp
@@ -9,6 +9,7 @@
 #include "HlslLexer.h"
 #include "HlslParser.h"
 #include "ShaderPreprocessor.h"
+#include "ShaderCompilerCommon.h"
 
 #include "RequiredProgramMainCPPInclude.h"
 
@@ -169,6 +170,21 @@ namespace CCT
 			}
 		}
 
+		{
+			const ANSICHAR* USFSource = ShaderSource;
+			CrossCompiler::FHlslccHeader CCHeader;
+			int32 Len = FCStringAnsi::Strlen(USFSource);
+			if (!CCHeader.Read(USFSource, Len))
+			{
+				UE_LOG(LogCrossCompilerTool, Error, TEXT("Bad hlslcc header found"));
+			}
+
+			if (*USFSource != '#')
+			{
+				UE_LOG(LogCrossCompilerTool, Error, TEXT("Bad hlslcc header found! Missing '#'!"));
+			}
+		}
+
 		free(ShaderSource);
 		free(ErrorLog);
 
diff --git a/Engine/Source/Runtime/IOS/MetalRHI/Private/MetalShaders.cpp b/Engine/Source/Runtime/IOS/MetalRHI/Private/MetalShaders.cpp
--- a/Engine/Source/Runtime/IOS/MetalRHI/Private/MetalShaders.cpp
+++ b/Engine/Source/Runtime/IOS/MetalRHI/Private/MetalShaders.cpp
@@ -515,7 +515,7 @@ void FMetalShaderParameterCache::CommitPackedGlobals(int32 Stage, const FMetalSh
 	}
 }
 
-void FMetalShaderParameterCache::CommitPackedUniformBuffers(TRefCountPtr<FMetalBoundShaderState> BoundShaderState, int32 Stage, const TArray< TRefCountPtr<FRHIUniformBuffer> >& RHIUniformBuffers, const TArray<FMetalUniformBufferCopyInfo>& UniformBuffersCopyInfo)
+void FMetalShaderParameterCache::CommitPackedUniformBuffers(TRefCountPtr<FMetalBoundShaderState> BoundShaderState, int32 Stage, const TArray< TRefCountPtr<FRHIUniformBuffer> >& RHIUniformBuffers, const TArray<CrossCompiler::FUniformBufferCopyInfo>& UniformBuffersCopyInfo)
 {
 //	SCOPE_CYCLE_COUNTER(STAT_MetalConstantBufferUpdateTime);
 	// Uniform Buffers are split into precision/type; the list of RHI UBs is traversed and if a new one was set, its
@@ -538,7 +538,7 @@ void FMetalShaderParameterCache::CommitPackedUniformBuffers(TRefCountPtr<FMetalB
 			const uint32* RESTRICT SourceData = (uint32*)((uint8*)[EmulatedUniformBuffer->Buffer contents] + EmulatedUniformBuffer->Offset);//->Data.GetTypedData();
 			for (int32 InfoIndex = LastInfoIndex; InfoIndex < UniformBuffersCopyInfo.Num(); ++InfoIndex)
 			{
-				const FMetalUniformBufferCopyInfo& Info = UniformBuffersCopyInfo[InfoIndex];
+				const CrossCompiler::FUniformBufferCopyInfo& Info = UniformBuffersCopyInfo[InfoIndex];
 				if (Info.SourceUBIndex == BufferIndex)
 				{
 					float* RESTRICT ScratchMem = (float*)PackedGlobalUniforms[Info.DestUBTypeIndex];
diff --git a/Engine/Source/Runtime/IOS/MetalRHI/Public/MetalShaderResources.h b/Engine/Source/Runtime/IOS/MetalRHI/Public/MetalShaderResources.h
--- a/Engine/Source/Runtime/IOS/MetalRHI/Public/MetalShaderResources.h
+++ b/Engine/Source/Runtime/IOS/MetalRHI/Public/MetalShaderResources.h
@@ -133,38 +133,12 @@ inline FArchive& operator<<(FArchive& Ar, FMetalShaderBindings& Bindings)
 	return Ar;
 }
 
-// Information for copying members from uniform buffers to packed
-struct FMetalUniformBufferCopyInfo
-{
-	uint16 SourceOffsetInFloats;
-	uint8 SourceUBIndex;
-	uint8 DestUBIndex;
-	uint8 DestUBTypeName;
-	uint8 DestUBTypeIndex;
-	uint16 DestOffsetInFloats;
-	uint16 SizeInFloats;
-};
-
-inline FArchive& operator<<(FArchive& Ar, FMetalUniformBufferCopyInfo& Info)
-{
-	Ar << Info.SourceOffsetInFloats;
-	Ar << Info.SourceUBIndex;
-	Ar << Info.DestUBIndex;
-	Ar << Info.DestUBTypeName;
-	if (Ar.IsLoading())
-	{
-		Info.DestUBTypeIndex = CrossCompiler::PackedTypeNameToTypeIndex(Info.DestUBTypeName);
-	}
-	Ar << Info.DestOffsetInFloats;
-	Ar << Info.SizeInFloats;
-	return Ar;
-}
-
 struct FMetalCodeHeader
 {
 	uint32 Frequency;
 	FMetalShaderBindings Bindings;
-	TArray<FMetalUniformBufferCopyInfo> UniformBuffersCopyInfo;
+	TArray<CrossCompiler::FUniformBufferCopyInfo> UniformBuffersCopyInfo;
+	TArray<TCHAR> ShaderName;
 
 	uint8 NumThreadsX;
 	uint8 NumThreadsY;
@@ -190,7 +164,7 @@ inline FArchive& operator<<(FArchive& Ar, FMetalCodeHeader& Header)
 		Header.UniformBuffersCopyInfo.Empty(NumInfos);
 		for (int32 Index = 0; Index < NumInfos; ++Index)
 		{
-			FMetalUniformBufferCopyInfo Info;
+			CrossCompiler::FUniformBufferCopyInfo Info;
 			Ar << Info;
 			Header.UniformBuffersCopyInfo.Add(Info);
 		}
diff --git a/Engine/Source/Runtime/OpenGLDrv/Private/OpenGLShaders.cpp b/Engine/Source/Runtime/OpenGLDrv/Private/OpenGLShaders.cpp
--- a/Engine/Source/Runtime/OpenGLDrv/Private/OpenGLShaders.cpp
+++ b/Engine/Source/Runtime/OpenGLDrv/Private/OpenGLShaders.cpp
@@ -2669,7 +2669,7 @@ void FOpenGLShaderParameterCache::CommitPackedGlobals(const FOpenGLLinkedProgram
 	}
 }
 
-void FOpenGLShaderParameterCache::CommitPackedUniformBuffers(FOpenGLLinkedProgram* LinkedProgram, int32 Stage, FUniformBufferRHIRef* RHIUniformBuffers, const TArray<FOpenGLUniformBufferCopyInfo>& UniformBuffersCopyInfo)
+void FOpenGLShaderParameterCache::CommitPackedUniformBuffers(FOpenGLLinkedProgram* LinkedProgram, int32 Stage, FUniformBufferRHIRef* RHIUniformBuffers, const TArray<CrossCompiler::FUniformBufferCopyInfo>& UniformBuffersCopyInfo)
 {
 	SCOPE_CYCLE_COUNTER(STAT_OpenGLConstantBufferUpdateTime);
 	VERIFY_GL_SCOPE();
@@ -2689,7 +2689,7 @@ void FOpenGLShaderParameterCache::CommitPackedUniformBuffers(FOpenGLLinkedProgra
 			const uint32* RESTRICT SourceData = UniformBuffer->EmulatedBufferData->Data.GetData();
 			for (int32 InfoIndex = LastInfoIndex; InfoIndex < UniformBuffersCopyInfo.Num(); ++InfoIndex)
 			{
-				const FOpenGLUniformBufferCopyInfo& Info = UniformBuffersCopyInfo[InfoIndex];
+				const CrossCompiler::FUniformBufferCopyInfo& Info = UniformBuffersCopyInfo[InfoIndex];
 				if (Info.SourceUBIndex == BufferIndex)
 				{
 					check((Info.DestOffsetInFloats + Info.SizeInFloats) * sizeof(float) <= (uint32)GlobalUniformArraySize);
@@ -2722,7 +2722,7 @@ void FOpenGLShaderParameterCache::CommitPackedUniformBuffers(FOpenGLLinkedProgra
 				// Go through the list of copy commands and perform the appropriate copy into the scratch buffer
 				for (int32 InfoIndex = LastCopyInfoIndex; InfoIndex < UniformBuffersCopyInfo.Num(); ++InfoIndex)
 				{
-					const FOpenGLUniformBufferCopyInfo& Info = UniformBuffersCopyInfo[InfoIndex];
+					const CrossCompiler::FUniformBufferCopyInfo& Info = UniformBuffersCopyInfo[InfoIndex];
 					if (Info.SourceUBIndex == BufferIndex)
 					{
 						const uint32* RESTRICT SourceData = UniformBuffer->EmulatedBufferData->Data.GetData();
diff --git a/Engine/Source/Runtime/OpenGLDrv/Public/OpenGLShaderResources.h b/Engine/Source/Runtime/OpenGLDrv/Public/OpenGLShaderResources.h
--- a/Engine/Source/Runtime/OpenGLDrv/Public/OpenGLShaderResources.h
+++ b/Engine/Source/Runtime/OpenGLDrv/Public/OpenGLShaderResources.h
@@ -201,33 +201,6 @@ inline FArchive& operator<<(FArchive& Ar, FOpenGLShaderBindings& Bindings)
 	return Ar;
 }
 
-// Information for copying members from uniform buffers to packed
-struct FOpenGLUniformBufferCopyInfo
-{
-	uint16 SourceOffsetInFloats;
-	uint8 SourceUBIndex;
-	uint8 DestUBIndex;
-	uint8 DestUBTypeName;
-	uint8 DestUBTypeIndex;
-	uint16 DestOffsetInFloats;
-	uint16 SizeInFloats;
-};
-
-inline FArchive& operator<<(FArchive& Ar, FOpenGLUniformBufferCopyInfo& Info)
-{
-	Ar << Info.SourceOffsetInFloats;
-	Ar << Info.SourceUBIndex;
-	Ar << Info.DestUBIndex;
-	Ar << Info.DestUBTypeName;
-	if (Ar.IsLoading())
-	{
-		Info.DestUBTypeIndex = CrossCompiler::PackedTypeNameToTypeIndex(Info.DestUBTypeName);
-	}
-	Ar << Info.DestOffsetInFloats;
-	Ar << Info.SizeInFloats;
-	return Ar;
-}
-
 /**
  * Code header information.
  */
@@ -236,7 +209,8 @@ struct FOpenGLCodeHeader
 	uint32 GlslMarker;
 	uint16 FrequencyMarker;
 	FOpenGLShaderBindings Bindings;
-	TArray<FOpenGLUniformBufferCopyInfo> UniformBuffersCopyInfo;
+	TArray<TCHAR> ShaderName;
+	TArray<CrossCompiler::FUniformBufferCopyInfo> UniformBuffersCopyInfo;
 };
 
 inline FArchive& operator<<(FArchive& Ar, FOpenGLCodeHeader& Header)
@@ -258,7 +232,7 @@ inline FArchive& operator<<(FArchive& Ar, FOpenGLCodeHeader& Header)
 		Header.UniformBuffersCopyInfo.Empty(NumInfos);
 		for (int32 Index = 0; Index < NumInfos; ++Index)
 		{
-			FOpenGLUniformBufferCopyInfo Info;
+			CrossCompiler::FUniformBufferCopyInfo Info;
 			Ar << Info;
 			Header.UniformBuffersCopyInfo.Add(Info);
 		}
@@ -290,7 +264,7 @@ public:
 	FOpenGLShaderBindings Bindings;
 
 	// List of memory copies from RHIUniformBuffer to packed uniforms
-	TArray<FOpenGLUniformBufferCopyInfo> UniformBuffersCopyInfo;
+	TArray<CrossCompiler::FUniformBufferCopyInfo> UniformBuffersCopyInfo;
 
 #if DEBUG_GL_SHADERS
 	TArray<ANSICHAR> GlslCode;
@@ -370,7 +344,7 @@ public:
 	 */
 	void CommitPackedGlobals(const FOpenGLLinkedProgram* LinkedProgram, int32 Stage);
 
-	void CommitPackedUniformBuffers(FOpenGLLinkedProgram* LinkedProgram, int32 Stage, FUniformBufferRHIRef* UniformBuffers, const TArray<FOpenGLUniformBufferCopyInfo>& UniformBuffersCopyInfo);
+	void CommitPackedUniformBuffers(FOpenGLLinkedProgram* LinkedProgram, int32 Stage, FUniformBufferRHIRef* UniformBuffers, const TArray<CrossCompiler::FUniformBufferCopyInfo>& UniformBuffersCopyInfo);
 
 private:
 	/** CPU memory block for storing uniform values. */
diff --git a/Engine/Source/Runtime/ShaderCore/Public/CrossCompilerCommon.h b/Engine/Source/Runtime/ShaderCore/Public/CrossCompilerCommon.h
--- a/Engine/Source/Runtime/ShaderCore/Public/CrossCompilerCommon.h
+++ b/Engine/Source/Runtime/ShaderCore/Public/CrossCompilerCommon.h
@@ -79,4 +79,44 @@ namespace CrossCompiler
 		Ar << Info.TypeIndex;
 		return Ar;
 	}
+
+	struct FShaderBindings
+	{
+		TArray<TArray<FPackedArrayInfo>>	PackedUniformBuffers;
+		TArray<FPackedArrayInfo>			PackedGlobalArrays;
+		FShaderResourceTable				ShaderResourceTable;
+
+		uint16	InOutMask;
+		uint8	NumSamplers;
+		uint8	NumUniformBuffers;
+		uint8	NumUAVs;
+		bool	bHasRegularUniformBuffers;
+	};
+
+	// Information for copying members from uniform buffers to packed
+	struct FUniformBufferCopyInfo
+	{
+		uint16 SourceOffsetInFloats;
+		uint8 SourceUBIndex;
+		uint8 DestUBIndex;
+		uint8 DestUBTypeName;
+		uint8 DestUBTypeIndex;
+		uint16 DestOffsetInFloats;
+		uint16 SizeInFloats;
+	};
+
+	inline FArchive& operator<<(FArchive& Ar, FUniformBufferCopyInfo& Info)
+	{
+		Ar << Info.SourceOffsetInFloats;
+		Ar << Info.SourceUBIndex;
+		Ar << Info.DestUBIndex;
+		Ar << Info.DestUBTypeName;
+		if (Ar.IsLoading())
+		{
+			Info.DestUBTypeIndex = CrossCompiler::PackedTypeNameToTypeIndex(Info.DestUBTypeName);
+		}
+		Ar << Info.DestOffsetInFloats;
+		Ar << Info.SizeInFloats;
+		return Ar;
+	}
 }
