ES 3.1 GLSL fixes (Rivalry now works on Nexus 9 and Shield Tablet 8)

diff --git a/Engine/Source/Developer/ShaderFormatOpenGL/Private/GlslBackend.cpp b/Engine/Source/Developer/ShaderFormatOpenGL/Private/GlslBackend.cpp
--- a/Engine/Source/Developer/ShaderFormatOpenGL/Private/GlslBackend.cpp
+++ b/Engine/Source/Developer/ShaderFormatOpenGL/Private/GlslBackend.cpp
@@ -749,6 +749,67 @@ class ir_gen_glsl_visitor : public ir_visitor
 		}
 	}
 
+	enum EPrecisionModifier
+	{
+		GLSL_PRECISION_DEFAULT,
+		GLSL_PRECISION_LOWP,
+		GLSL_PRECISION_MEDIUMP,
+		GLSL_PRECISION_HIGHP,
+	};
+
+	EPrecisionModifier GetPrecisionModifier(const struct glsl_type *type)
+	{
+		if (type->is_sampler() || type->is_image())
+		{
+			if (GDefaultPrecisionIsHalf && type->inner_type->base_type == GLSL_TYPE_FLOAT)
+			{
+				return GLSL_PRECISION_HIGHP;
+			}
+			else if (!GDefaultPrecisionIsHalf && type->inner_type->base_type == GLSL_TYPE_HALF)
+			{
+				return GLSL_PRECISION_MEDIUMP;
+			}
+			else // shadow samplers, integer textures etc
+			{
+				return GLSL_PRECISION_HIGHP;
+			}
+		}
+		else if (GDefaultPrecisionIsHalf && (type->base_type == GLSL_TYPE_FLOAT || (type->is_array() && type->element_type()->base_type == GLSL_TYPE_FLOAT)))
+		{
+			return GLSL_PRECISION_HIGHP;
+		}
+		else if (!GDefaultPrecisionIsHalf && (type->base_type == GLSL_TYPE_HALF || (type->is_array() && type->element_type()->base_type == GLSL_TYPE_HALF)))
+		{
+			return GLSL_PRECISION_MEDIUMP;
+		}
+		else if (type->is_integer())
+		{
+			return GLSL_PRECISION_HIGHP;
+		}
+		return GLSL_PRECISION_DEFAULT;
+	}
+
+	void AppendPrecisionModifier(char** buffer, EPrecisionModifier PrecisionModifier)
+	{
+		switch (PrecisionModifier)
+		{
+			case GLSL_PRECISION_LOWP:
+				ralloc_asprintf_append(buffer, "lowp ");
+				break;
+			case GLSL_PRECISION_MEDIUMP:
+				ralloc_asprintf_append(buffer, "mediump ");
+				break;
+			case GLSL_PRECISION_HIGHP:
+				ralloc_asprintf_append(buffer, "highp ");
+				break;
+			case GLSL_PRECISION_DEFAULT:
+				break;
+			default:
+				// we missed a type
+				check(false);
+		}
+	}
+
 	/**
 	* \name Visit methods
 	*
@@ -887,6 +948,14 @@ class ir_gen_glsl_visitor : public ir_visitor
 					const glsl_struct_field* field = &inner_type->fields.structure[0];
 					check(strcmp(field->name,"Data")==0);
 					
+					if (bEmitPrecision)
+					{
+						if (field->type->is_integer())
+						{
+							ralloc_asprintf_append(buffer, "flat ");
+						}
+						AppendPrecisionModifier(buffer, GetPrecisionModifier(field->type));
+					}
 					print_type_pre(field->type);
 					ralloc_asprintf_append(buffer, ", Data");
 					print_type_post(field->type);
@@ -969,23 +1038,7 @@ class ir_gen_glsl_visitor : public ir_visitor
 
 				if (bEmitPrecision)
 				{
-					if (GDefaultPrecisionIsHalf && var->type->inner_type->base_type == GLSL_TYPE_FLOAT)
-					{
-						ralloc_asprintf_append(buffer, "highp ");
-					}
-					else if (!GDefaultPrecisionIsHalf && var->type->inner_type->base_type == GLSL_TYPE_HALF)
-					{
-						ralloc_asprintf_append(buffer, "mediump ");
-					}
-					else if (var->type->inner_type->is_integer())
-					{
-						ralloc_asprintf_append(buffer, "highp ");
-					}
-					else
-					{
-						// we missed a type
-						check(false);
-					}
+					AppendPrecisionModifier(buffer, GetPrecisionModifier(var->type));
 				}
 				print_type_pre(var->type);
 			}
@@ -1012,33 +1065,7 @@ class ir_gen_glsl_visitor : public ir_visitor
 					);
 				if (bEmitPrecision)
 				{
-					if (var->type->is_sampler())
-					{
-						if (GDefaultPrecisionIsHalf && var->type->inner_type->base_type == GLSL_TYPE_FLOAT)
-						{
-							ralloc_asprintf_append(buffer, "highp ");
-						}
-						else if (!GDefaultPrecisionIsHalf && var->type->inner_type->base_type == GLSL_TYPE_HALF)
-						{
-							ralloc_asprintf_append(buffer, "mediump ");
-						}
-						else // shadow samplers, integer textures etc
-						{
-							ralloc_asprintf_append(buffer, "highp ");
-						}
-					}
-					else if (GDefaultPrecisionIsHalf && (var->type->base_type == GLSL_TYPE_FLOAT || (var->type->is_array() && var->type->element_type()->base_type == GLSL_TYPE_FLOAT)))
-					{
-						ralloc_asprintf_append(buffer, "highp ");
-					}
-					else if (!GDefaultPrecisionIsHalf && (var->type->base_type == GLSL_TYPE_HALF || (var->type->is_array() && var->type->element_type()->base_type == GLSL_TYPE_HALF)))
-					{
-						ralloc_asprintf_append(buffer, "mediump ");
-					}
-					else if (var->type->is_integer())
-					{
-						ralloc_asprintf_append(buffer, "highp ");
-					}
+					AppendPrecisionModifier(buffer, GetPrecisionModifier(var->type));
 				}
 
 				if (bGenerateLayoutLocations && var->explicit_location)
diff --git a/Engine/Source/Developer/ShaderFormatOpenGL/Private/ShaderFormatOpenGL.cpp b/Engine/Source/Developer/ShaderFormatOpenGL/Private/ShaderFormatOpenGL.cpp
--- a/Engine/Source/Developer/ShaderFormatOpenGL/Private/ShaderFormatOpenGL.cpp
+++ b/Engine/Source/Developer/ShaderFormatOpenGL/Private/ShaderFormatOpenGL.cpp
@@ -34,7 +34,7 @@ class FShaderFormatGLSL : public IShaderFormat
 		UE_SHADER_GLSL_150ES3_1_VER = 3,
 		UE_SHADER_GLSL_ES2_VER_WEBGL  = 17,
 		UE_SHADER_GLSL_ES2_IOS_VER  = 6,
-		UE_SHADER_GLSL_310_ES_EXT_VER = 5,
+		UE_SHADER_GLSL_310_ES_EXT_VER = 6,
 	}; 
 
 	void CheckFormat(FName Format) const
diff --git a/Engine/Source/Runtime/OpenGLDrv/Private/OpenGLShaders.cpp b/Engine/Source/Runtime/OpenGLDrv/Private/OpenGLShaders.cpp
--- a/Engine/Source/Runtime/OpenGLDrv/Private/OpenGLShaders.cpp
+++ b/Engine/Source/Runtime/OpenGLDrv/Private/OpenGLShaders.cpp
@@ -552,9 +552,12 @@ ShaderType* CompileOpenGLShader(const TArray<uint8>& Code)
 			{
 #if PLATFORM_DESKTOP
 				AppendCString(GlslCode, "#extension GL_ARB_separate_shader_objects : enable\n");
-#endif
 				AppendCString(GlslCode, "#define INTERFACE_LOCATION(Pos) layout(location=Pos) \n");
 				AppendCString(GlslCode, "#define INTERFACE_BLOCK(Pos, Interp, Modifiers, Semantic, PreType, PostType) layout(location=Pos) Interp Modifiers struct { PreType PostType; }\n");
+#else
+				AppendCString(GlslCode, "#define INTERFACE_LOCATION(Pos) layout(location=Pos) \n");
+				AppendCString(GlslCode, "#define INTERFACE_BLOCK(Pos, Interp, Modifiers, Semantic, PreType, PostType) layout(location=Pos) Modifiers Semantic { PreType PostType; }\n");
+#endif
 			}
 			else
 			{
