#ue4-buildwatcher - Fix shadow variable warning.

diff --git a/Engine/Source/Developer/ShaderFormatOpenGL/Private/GlslBackend.cpp b/Engine/Source/Developer/ShaderFormatOpenGL/Private/GlslBackend.cpp
--- a/Engine/Source/Developer/ShaderFormatOpenGL/Private/GlslBackend.cpp
+++ b/Engine/Source/Developer/ShaderFormatOpenGL/Private/GlslBackend.cpp
@@ -789,18 +789,18 @@ class ir_gen_glsl_visitor : public ir_visitor
 		return GLSL_PRECISION_DEFAULT;
 	}
 
-	void AppendPrecisionModifier(char** buffer, EPrecisionModifier PrecisionModifier)
+	void AppendPrecisionModifier(char** inBuffer, EPrecisionModifier PrecisionModifier)
 	{
 		switch (PrecisionModifier)
 		{
 			case GLSL_PRECISION_LOWP:
-				ralloc_asprintf_append(buffer, "lowp ");
+				ralloc_asprintf_append(inBuffer, "lowp ");
 				break;
 			case GLSL_PRECISION_MEDIUMP:
-				ralloc_asprintf_append(buffer, "mediump ");
+				ralloc_asprintf_append(inBuffer, "mediump ");
 				break;
 			case GLSL_PRECISION_HIGHP:
-				ralloc_asprintf_append(buffer, "highp ");
+				ralloc_asprintf_append(inBuffer, "highp ");
 				break;
 			case GLSL_PRECISION_DEFAULT:
 				break;
