Linux: some fixes/improvements for shader debugging.

diff --git a/Engine/Source/Developer/ShaderCompilerCommon/Private/ShaderCompilerCommon.cpp b/Engine/Source/Developer/ShaderCompilerCommon/Private/ShaderCompilerCommon.cpp
--- a/Engine/Source/Developer/ShaderCompilerCommon/Private/ShaderCompilerCommon.cpp
+++ b/Engine/Source/Developer/ShaderCompilerCommon/Private/ShaderCompilerCommon.cpp
@@ -357,7 +357,8 @@ namespace CrossCompiler
 		}
 		else if (PLATFORM_LINUX)
 		{
-			BatchFile = FPaths::RootDir() / FString::Printf(TEXT("Engine/Binaries/Linux/CrossCompilerTool %s -o=%s %s -entry=%s %s %s"), *ShaderFile, *OutputFile, *FrequencySwitch, *EntryPoint, *VersionSwitch, *CCTCmdLine);
+			BatchFile = TEXT("#!/bin/sh\n");
+			BatchFile += FPaths::RootDir() / FString::Printf(TEXT("Engine/Binaries/Linux/CrossCompilerTool %s -o=%s %s -entry=%s %s %s"), *ShaderFile, *OutputFile, FrequencySwitch, *EntryPoint, *VersionSwitch, *CCTCmdLine);
 		}
 		else if (PLATFORM_WINDOWS)
 		{
diff --git a/Engine/Source/Developer/ShaderFormatOpenGL/Private/OpenGLShaderCompiler.cpp b/Engine/Source/Developer/ShaderFormatOpenGL/Private/OpenGLShaderCompiler.cpp
--- a/Engine/Source/Developer/ShaderFormatOpenGL/Private/OpenGLShaderCompiler.cpp
+++ b/Engine/Source/Developer/ShaderFormatOpenGL/Private/OpenGLShaderCompiler.cpp
@@ -1315,7 +1315,8 @@ void CompileShader_Windows_OGL(const FShaderCompilerInput& Input,FShaderCompiler
 			const FString CCBatchFileContents = CreateCrossCompilerBatchFile(USFFile, GLSLFile, *Input.EntryPointName, Frequency, Version, CCFlags);
 			if (!CCBatchFileContents.IsEmpty())
 			{
-				FFileHelper::SaveStringToFile(CCBatchFileContents, *(Input.DumpDebugInfoPath / TEXT("CrossCompile.bat")));
+				const TCHAR * ScriptName = PLATFORM_WINDOWS ? TEXT("CrossCompile.bat") : TEXT("CrossCompile.sh");
+				FFileHelper::SaveStringToFile(CCBatchFileContents, *(Input.DumpDebugInfoPath / ScriptName));
 			}
 		}
 
