Merging changes from ZebraBranch that:
- Append the source shader name to the GLSL when debug info is enabled to make it easier to debug.
- Always preserves GLSL flow control for Mac GLSL to improve performance on 2 of 3 vendors.
- Exposes a CVar to control whether to preserve or avoid GLSL flow-control for other shader platforms. The default remains to avoid flow control in GLSL.

diff --git a/Engine/Source/Developer/ShaderFormatOpenGL/Private/OpenGLShaderCompiler.cpp b/Engine/Source/Developer/ShaderFormatOpenGL/Private/OpenGLShaderCompiler.cpp
--- a/Engine/Source/Developer/ShaderFormatOpenGL/Private/OpenGLShaderCompiler.cpp
+++ b/Engine/Source/Developer/ShaderFormatOpenGL/Private/OpenGLShaderCompiler.cpp
@@ -1169,6 +1169,7 @@ void CompileShader_Windows_OGL(const FShaderCompilerInput& Input,FShaderCompiler
 	FString PreprocessedShader;
 	FShaderCompilerDefinitions AdditionalDefines;
 	EHlslCompileTarget HlslCompilerTarget = HCT_InvalidTarget;
+	ECompilerFlags PlatformFlowControl = CFLAG_AvoidFlowControl;
 
 	AdditionalDefines.SetDefine(TEXT("COMPILER_HLSLCC"), 1);
 	switch (Version)
@@ -1197,6 +1198,9 @@ void CompileShader_Windows_OGL(const FShaderCompilerInput& Input,FShaderCompiler
 			AdditionalDefines.SetDefine(TEXT("COMPILER_GLSL"), 1);
 			AdditionalDefines.SetDefine(TEXT("GL3_PROFILE"), 1);
 			HlslCompilerTarget = HCT_FeatureLevelSM4;
+			// On OS X it is always better to leave the flow control statements in the GLSL & let the GLSL->GPU compilers
+			// optimise it appropriately for each GPU. This gives a performance gain on AMD & Intel and is neutral on Nvidia.
+			PlatformFlowControl = CFLAG_PreferFlowControl;
 			break;
 
 		case GLSL_ES2_WEBGL:
@@ -1243,7 +1247,14 @@ void CompileShader_Windows_OGL(const FShaderCompilerInput& Input,FShaderCompiler
 	
 	const bool bDumpDebugInfo = (Input.DumpDebugInfoPath != TEXT("") && IFileManager::Get().DirectoryExists(*Input.DumpDebugInfoPath));
 
-	AdditionalDefines.SetDefine(TEXT("COMPILER_SUPPORTS_ATTRIBUTES"), (uint32)1);
+	if(Input.Environment.CompilerFlags.Contains(CFLAG_AvoidFlowControl) && PlatformFlowControl == CFLAG_AvoidFlowControl)
+	{
+		AdditionalDefines.SetDefine(TEXT("COMPILER_SUPPORTS_ATTRIBUTES"), (uint32)1);
+	}
+	else
+	{
+		AdditionalDefines.SetDefine(TEXT("COMPILER_SUPPORTS_ATTRIBUTES"), (uint32)0);
+	}
 	if (PreprocessShader(PreprocessedShader, Output, Input, AdditionalDefines))
 	{
 		char* GlslShaderSource = NULL;
@@ -1361,6 +1372,13 @@ void CompileShader_Windows_OGL(const FShaderCompilerInput& Input,FShaderCompiler
 
 				if (GlslSourceLen > 0)
 				{
+					uint32 Len = FCStringAnsi::Strlen(TCHAR_TO_ANSI(*Input.SourceFilename)) + FCStringAnsi::Strlen(TCHAR_TO_ANSI(*Input.EntryPointName)) + FCStringAnsi::Strlen(GlslShaderSource) + 20;
+					char* Dest = (char*)malloc(Len);
+					FCStringAnsi::Snprintf(Dest, Len, "// ! %s.usf:%s\n%s", (const char*)TCHAR_TO_ANSI(*Input.SourceFilename), (const char*)TCHAR_TO_ANSI(*Input.EntryPointName), (const char*)GlslShaderSource);
+					free(GlslShaderSource);
+					GlslShaderSource = Dest;
+					GlslSourceLen = FCStringAnsi::Strlen(GlslShaderSource);
+					
 					FArchive* FileWriter = IFileManager::Get().CreateFileWriter(*(Input.DumpDebugInfoPath / Input.SourceFilename + TEXT(".glsl")));
 					if (FileWriter)
 					{
diff --git a/Engine/Source/Runtime/Engine/Private/ShaderCompiler/ShaderCompiler.cpp b/Engine/Source/Runtime/Engine/Private/ShaderCompiler/ShaderCompiler.cpp
--- a/Engine/Source/Runtime/Engine/Private/ShaderCompiler/ShaderCompiler.cpp
+++ b/Engine/Source/Runtime/Engine/Private/ShaderCompiler/ShaderCompiler.cpp
@@ -53,6 +53,13 @@ static TAutoConsoleVariable<int32> CVarOptimizeShaders(
 	TEXT("Whether to optimize shaders.  When using graphical debuggers like Nsight it can be useful to disable this on startup."),
 	ECVF_ReadOnly);
 
+static TAutoConsoleVariable<int32> CVarAvoidFlowControl(
+	TEXT("r.Shaders.AvoidFlowControl"),
+	1,
+	TEXT("When enabled reduces flow-control in favour of longer shaders with greater register usage. When disabled shader flow-control is preserved for platform runtime to optimise which may be beneficial on register constrained GPUs.\n")
+	TEXT("Defaults to 1 (enabled)."),
+	ECVF_ReadOnly);
+
 #if !(UE_BUILD_SHIPPING || UE_BUILD_TEST)
 static TAutoConsoleVariable<FString> CVarD3DCompilerPath(TEXT("r.D3DCompilerPath"),
 	TEXT(""),	// default
@@ -1892,6 +1899,15 @@ void GlobalBeginCompileShader(
 			Input.Environment.CompilerFlags.Add(CFLAG_KeepDebugInfo);
 		}
 	}
+	
+	{
+		static const auto CVar = IConsoleManager::Get().FindConsoleVariable(TEXT("r.Shaders.AvoidFlowControl"));
+		
+		if (CVar->GetInt() != 0)
+		{
+			Input.Environment.CompilerFlags.Add(CFLAG_AvoidFlowControl);
+		}
+	}
 
 	Input.Environment.SetDefine(TEXT("HAS_INVERTED_Z_BUFFER"), ((int32)ERHIZBuffer::IsInverted != 0) ? 1 : 0);
 
diff --git a/Engine/Source/Runtime/OpenGLDrv/Private/OpenGLShaders.cpp b/Engine/Source/Runtime/OpenGLDrv/Private/OpenGLShaders.cpp
--- a/Engine/Source/Runtime/OpenGLDrv/Private/OpenGLShaders.cpp
+++ b/Engine/Source/Runtime/OpenGLDrv/Private/OpenGLShaders.cpp
@@ -564,6 +564,13 @@ ShaderType* CompileOpenGLShader(const TArray<uint8>& Code)
 				AppendCString(GlslCode, "#define INTERFACE_LOCATION(Pos) \n");
 				AppendCString(GlslCode, "#define INTERFACE_BLOCK(Pos, Interp, Modifiers, Semantic, PreType, PostType) Modifiers Semantic { Interp PreType PostType; }\n");
 			}
+			
+			if(Header.ShaderName.Num())
+			{
+				AppendCString(GlslCode, "// ");
+				AppendCString(GlslCode, TCHAR_TO_ANSI(Header.ShaderName.GetData()));
+				AppendCString(GlslCode, "\n");
+			}
 		}
 
 #if PLATFORM_ANDROID 
diff --git a/Engine/Source/Runtime/OpenGLDrv/Public/OpenGLShaderResources.h b/Engine/Source/Runtime/OpenGLDrv/Public/OpenGLShaderResources.h
--- a/Engine/Source/Runtime/OpenGLDrv/Public/OpenGLShaderResources.h
+++ b/Engine/Source/Runtime/OpenGLDrv/Public/OpenGLShaderResources.h
@@ -218,6 +218,7 @@ inline FArchive& operator<<(FArchive& Ar, FOpenGLCodeHeader& Header)
 	Ar << Header.GlslMarker;
 	Ar << Header.FrequencyMarker;
 	Ar << Header.Bindings;
+	Ar << Header.ShaderName;
 	int32 NumInfos = Header.UniformBuffersCopyInfo.Num();
 	Ar << NumInfos;
 	if (Ar.IsSaving())
