Fix FOpenGLCodeHeader serialisation for android. (where TCHAR == 4 bytes)

diff --git a/Engine/Source/Developer/ShaderFormatOpenGL/Private/OpenGLShaderCompiler.cpp b/Engine/Source/Developer/ShaderFormatOpenGL/Private/OpenGLShaderCompiler.cpp
--- a/Engine/Source/Developer/ShaderFormatOpenGL/Private/OpenGLShaderCompiler.cpp
+++ b/Engine/Source/Developer/ShaderFormatOpenGL/Private/OpenGLShaderCompiler.cpp
@@ -799,7 +799,7 @@ static void BuildShaderOutput(
 			);
 	}
 
-	Header.ShaderName = CCHeader.Name.GetCharArray();
+	Header.ShaderName = CCHeader.Name;
 
 	// Build the SRT for this shader.
 	{
diff --git a/Engine/Source/Developer/ShaderFormatOpenGL/Private/ShaderFormatOpenGL.cpp b/Engine/Source/Developer/ShaderFormatOpenGL/Private/ShaderFormatOpenGL.cpp
--- a/Engine/Source/Developer/ShaderFormatOpenGL/Private/ShaderFormatOpenGL.cpp
+++ b/Engine/Source/Developer/ShaderFormatOpenGL/Private/ShaderFormatOpenGL.cpp
@@ -25,16 +25,16 @@ class FShaderFormatGLSL : public IShaderFormat
 	enum 
 	{
 		/** Version for shader format, this becomes part of the DDC key. */
-		UE_SHADER_GLSL_150_VER = 27,
-		UE_SHADER_GLSL_150_MAC_VER = 7,
-		UE_SHADER_GLSL_430_VER = 59,
-		UE_SHADER_GLSL_ES2_VER = 17,
-		UE_SHADER_GLSL_150ES2_VER = 19,
-		UE_SHADER_GLSL_150ES2NOUB_VER = 3,
-		UE_SHADER_GLSL_150ES3_1_VER = 3,
-		UE_SHADER_GLSL_ES2_VER_WEBGL  = 17,
-		UE_SHADER_GLSL_ES2_IOS_VER  = 6,
-		UE_SHADER_GLSL_310_ES_EXT_VER = 6,
+		UE_SHADER_GLSL_150_VER = 28,
+		UE_SHADER_GLSL_150_MAC_VER = 8,
+		UE_SHADER_GLSL_430_VER = 60,
+		UE_SHADER_GLSL_ES2_VER = 18,
+		UE_SHADER_GLSL_150ES2_VER = 20,
+		UE_SHADER_GLSL_150ES2NOUB_VER = 4,
+		UE_SHADER_GLSL_150ES3_1_VER = 4,
+		UE_SHADER_GLSL_ES2_VER_WEBGL  = 18,
+		UE_SHADER_GLSL_ES2_IOS_VER  = 7,
+		UE_SHADER_GLSL_310_ES_EXT_VER = 7,
 	}; 
 
 	void CheckFormat(FName Format) const
diff --git a/Engine/Source/Runtime/OpenGLDrv/Private/OpenGLShaders.cpp b/Engine/Source/Runtime/OpenGLDrv/Private/OpenGLShaders.cpp
--- a/Engine/Source/Runtime/OpenGLDrv/Private/OpenGLShaders.cpp
+++ b/Engine/Source/Runtime/OpenGLDrv/Private/OpenGLShaders.cpp
@@ -565,10 +565,10 @@ ShaderType* CompileOpenGLShader(const TArray<uint8>& Code)
 				AppendCString(GlslCode, "#define INTERFACE_BLOCK(Pos, Interp, Modifiers, Semantic, PreType, PostType) Modifiers Semantic { Interp PreType PostType; }\n");
 			}
 			
-			if(Header.ShaderName.Num())
+			if(Header.ShaderName.IsEmpty() == false)
 			{
 				AppendCString(GlslCode, "// ");
-				AppendCString(GlslCode, TCHAR_TO_ANSI(Header.ShaderName.GetData()));
+				AppendCString(GlslCode, TCHAR_TO_ANSI(Header.ShaderName.GetCharArray().GetData()));
 				AppendCString(GlslCode, "\n");
 			}
 		}
diff --git a/Engine/Source/Runtime/OpenGLDrv/Public/OpenGLShaderResources.h b/Engine/Source/Runtime/OpenGLDrv/Public/OpenGLShaderResources.h
--- a/Engine/Source/Runtime/OpenGLDrv/Public/OpenGLShaderResources.h
+++ b/Engine/Source/Runtime/OpenGLDrv/Public/OpenGLShaderResources.h
@@ -209,7 +209,7 @@ struct FOpenGLCodeHeader
 	uint32 GlslMarker;
 	uint16 FrequencyMarker;
 	FOpenGLShaderBindings Bindings;
-	TArray<TCHAR> ShaderName;
+	FString ShaderName;
 	TArray<CrossCompiler::FUniformBufferCopyInfo> UniformBuffersCopyInfo;
 };
 
