UE4 - Don't crash hlslcc if unhandled Nan value

diff --git a/Engine/Source/Developer/ShaderFormatOpenGL/Private/GlslBackend.cpp b/Engine/Source/Developer/ShaderFormatOpenGL/Private/GlslBackend.cpp
--- a/Engine/Source/Developer/ShaderFormatOpenGL/Private/GlslBackend.cpp
+++ b/Engine/Source/Developer/ShaderFormatOpenGL/Private/GlslBackend.cpp
@@ -1736,8 +1736,13 @@ class ir_gen_glsl_visitor : public ir_visitor
 						ralloc_asprintf_append(buffer, "(-1.0/0.0)");
 						break;
 
+					case 0x7fc00000u:
+						ralloc_asprintf_append(buffer, "(0.0/0.0) /*Real Nan*/");
+						break;
+
 					default:
-						check(0);
+						ralloc_asprintf_append(buffer, "Unhandled_Nan0x%08x", constant->value.u[index]);
+						break;
 				}
 			}
 		}
diff --git a/Engine/Source/Developer/ShaderFormatOpenGL/Private/ShaderFormatOpenGL.cpp b/Engine/Source/Developer/ShaderFormatOpenGL/Private/ShaderFormatOpenGL.cpp
--- a/Engine/Source/Developer/ShaderFormatOpenGL/Private/ShaderFormatOpenGL.cpp
+++ b/Engine/Source/Developer/ShaderFormatOpenGL/Private/ShaderFormatOpenGL.cpp
@@ -26,17 +26,17 @@ class FShaderFormatGLSL : public IShaderFormat
 	enum
 	{
 		/** Version for shader format, this becomes part of the DDC key. */
-		UE_SHADER_GLSL_150_VER = 60,
-		UE_SHADER_GLSL_150_MAC_VER = 60,
-		UE_SHADER_GLSL_430_VER = 61,
-		UE_SHADER_GLSL_ES2_VER = 60,
-		UE_SHADER_GLSL_150ES2_VER = 60,
-		UE_SHADER_GLSL_150ES2NOUB_VER = 60,
-		UE_SHADER_GLSL_150ES3_1_VER = 60,
-		UE_SHADER_GLSL_ES2_VER_WEBGL = 60,
-		UE_SHADER_GLSL_ES2_IOS_VER = 60,
-		UE_SHADER_GLSL_310_ES_EXT_VER = 61,
-		UE_SHADER_GLSL_ES3_1_ANDROID_VER = 60,
+		UE_SHADER_GLSL_150_VER = 61,
+		UE_SHADER_GLSL_150_MAC_VER = 61,
+		UE_SHADER_GLSL_430_VER = 62,
+		UE_SHADER_GLSL_ES2_VER = 61,
+		UE_SHADER_GLSL_150ES2_VER = 61,
+		UE_SHADER_GLSL_150ES2NOUB_VER = 61,
+		UE_SHADER_GLSL_150ES3_1_VER = 61,
+		UE_SHADER_GLSL_ES2_VER_WEBGL = 61,
+		UE_SHADER_GLSL_ES2_IOS_VER = 61,
+		UE_SHADER_GLSL_310_ES_EXT_VER = 62,
+		UE_SHADER_GLSL_ES3_1_ANDROID_VER = 61,
 	}; 
 
 	void CheckFormat(FName Format) const
