HTML5 - refraction shader

note: this CL also contains fixes to webgl2 [float4 vs half2] and a [% vs Mod()] material custom function changes to some TM-ShaderModels shaders
specifically: fixes to and similar with: DitherTemporalAA

#jria UE-46434  No Refraction in QA Game TM-Shadermodels HTML5
#rn

diff --git a/Engine/Config/BaseDeviceProfiles.ini b/Engine/Config/BaseDeviceProfiles.ini
--- a/Engine/Config/BaseDeviceProfiles.ini
+++ b/Engine/Config/BaseDeviceProfiles.ini
@@ -446,7 +446,7 @@ BaseProfileName=Android_Mali_T8xx
 [HTML5 DeviceProfile]
 DeviceType=HTML5
 BaseProfileName=
-+CVars=r.RefractionQuality=0
++CVars=r.RefractionQuality=1
 
 [Mac DeviceProfile]
 DeviceType=Mac
diff --git a/Engine/Shaders/Private/DistortAccumulatePS.usf b/Engine/Shaders/Private/DistortAccumulatePS.usf
--- a/Engine/Shaders/Private/DistortAccumulatePS.usf
+++ b/Engine/Shaders/Private/DistortAccumulatePS.usf
@@ -69,7 +69,7 @@ void Main(
 	// Sample depth at distortion offset
 	float2 DistortScreenUV = ScreenUV + Distortion;
 
-#if (ES2_PROFILE || ES3_1_PROFILE) 
+#if (ES2_PROFILE || ES3_1_PROFILE) && (!WEBGL)
 	float DistortSceneDepth = ConvertFromDeviceZ(Texture2DSampleLevel(SceneDepthTexture, SceneDepthTextureSampler, DistortScreenUV, 0).r);
 #else
 	float DistortSceneDepth = CalcSceneDepth(DistortScreenUV);
@@ -83,7 +83,7 @@ void Main(
 	float InvWidth = 1.0f / max(1.0f, Range );
 	Distortion *= saturate((Z - ZCompare) * InvWidth + Bias);
 
-#if (ES2_PROFILE || ES3_1_PROFILE) 
+#if (ES2_PROFILE || ES3_1_PROFILE)
 	// mobile path direclty outputs distorted colors
 	DistortScreenUV = ScreenUV + Distortion;
 	OutColor = CalcFullSceneColor(DistortScreenUV);
diff --git a/Engine/Shaders/Private/MobileBasePassPixelShader.usf b/Engine/Shaders/Private/MobileBasePassPixelShader.usf
--- a/Engine/Shaders/Private/MobileBasePassPixelShader.usf
+++ b/Engine/Shaders/Private/MobileBasePassPixelShader.usf
@@ -329,8 +329,13 @@ void Main(
 	FPixelMaterialInputs PixelMaterialInputs;
 	{
 		float4 ScreenPosition = SvPositionToResolvedScreenPosition(SvPosition);
+#if WEBGL
+		float3 WorldPosition = BasePassInterpolants.PixelPosition.xyz;
+		float3 WorldPositionExcludingWPO = BasePassInterpolants.PixelPosition.xyz;
+#else
 		float3 WorldPosition = BasePassInterpolants.PixelPosition;
 		float3 WorldPositionExcludingWPO = BasePassInterpolants.PixelPosition;
+#endif
 		#if USE_WORLD_POSITION_EXCLUDING_SHADER_OFFSETS
 			WorldPositionExcludingWPO = BasePassInterpolants.PixelPositionExcludingWPO;
 		#endif
diff --git a/Engine/Source/Developer/ShaderFormatOpenGL/Private/GlslBackend.cpp b/Engine/Source/Developer/ShaderFormatOpenGL/Private/GlslBackend.cpp
--- a/Engine/Source/Developer/ShaderFormatOpenGL/Private/GlslBackend.cpp
+++ b/Engine/Source/Developer/ShaderFormatOpenGL/Private/GlslBackend.cpp
@@ -108,6 +108,7 @@ static inline FCustomStdString FixHlslName(const glsl_type* Type, bool bIsES2)
 	{
 		return "mat4";
 	}
+#ifndef UE4_HTML5_TARGET_WEBGL2
 	else if (bIsES2 && Type->base_type == GLSL_TYPE_UINT)
 	{
 		// uint does not exist with GLSL 1.00 (ES2)
@@ -129,6 +130,7 @@ static inline FCustomStdString FixHlslName(const glsl_type* Type, bool bIsES2)
 			return "ivec4";
 		}
 	}
+#endif
 	return Name;
 }
 
diff --git a/Engine/Source/Developer/ShaderFormatOpenGL/Private/ShaderFormatOpenGL.cpp b/Engine/Source/Developer/ShaderFormatOpenGL/Private/ShaderFormatOpenGL.cpp
--- a/Engine/Source/Developer/ShaderFormatOpenGL/Private/ShaderFormatOpenGL.cpp
+++ b/Engine/Source/Developer/ShaderFormatOpenGL/Private/ShaderFormatOpenGL.cpp
@@ -176,14 +176,14 @@ public:
 		// the frontend will run the cross compiler
 		Frontend.CompileShader(Input, Output, WorkingDirectory, Version);
 
-		if (Version == GLSL_ES2 || Version == GLSL_150_ES2)
+		if (Version == GLSL_ES2 || Version == GLSL_150_ES2 || Version == GLSL_ES2_WEBGL)
 		{
 			if (Input.DumpDebugInfoPath != TEXT("") && IFileManager::Get().DirectoryExists(*Input.DumpDebugInfoPath))
 			{
 				FShaderCompilerInput ES2Input = Input;
 				ES2Input.DumpDebugInfoPath = ES2Input.DumpDebugInfoPath.Replace(
 					TEXT("GLSL_150_ES2"), 
-					(Version == GLSL_ES2 ? TEXT("GLSL_ES2") : TEXT("GLSL_ES2_150")),
+					(Version == GLSL_ES2 ? TEXT("GLSL_ES2") : Version == GLSL_150_ES2 ? TEXT("GLSL_ES2_150") : TEXT("GLSL_ES2_WEBGL")),
 					ESearchCase::CaseSensitive);
 				
 				if (!IFileManager::Get().DirectoryExists(*ES2Input.DumpDebugInfoPath))
@@ -192,7 +192,7 @@ public:
 				}
 
 				FShaderCompilerOutput ES2Output;
-				Frontend.CompileShader(ES2Input, ES2Output, WorkingDirectory, GLSL_ES2);
+				Frontend.CompileShader(ES2Input, ES2Output, WorkingDirectory, Version == GLSL_ES2_WEBGL ? GLSL_ES2_WEBGL : GLSL_ES2);
 			}
 		}
 	}
diff --git a/Engine/Source/Runtime/Renderer/Private/DistortionRendering.cpp b/Engine/Source/Runtime/Renderer/Private/DistortionRendering.cpp
--- a/Engine/Source/Runtime/Renderer/Private/DistortionRendering.cpp
+++ b/Engine/Source/Runtime/Renderer/Private/DistortionRendering.cpp
@@ -1169,10 +1169,12 @@ void FSceneRenderer::RenderDistortion(FRHICommandListImmediate& RHICmdList)
 void FSceneRenderer::RenderDistortionES2(FRHICommandListImmediate& RHICmdList)
 {
 	// We need access to HDR scene color
+#ifndef UE4_HTML5_TARGET_WEBGL2
 	if (!IsMobileHDR() || IsMobileHDRMosaic())
 	{
 		return;
 	}
+#endif
 
 	// do we need to render the distortion pass?
 	bool bRender=false;
