Fixed render target output index assignment for translated GLSL shaders

diff --git a/Engine/Source/Developer/ShaderFormatOpenGL/Private/GlslBackend.cpp b/Engine/Source/Developer/ShaderFormatOpenGL/Private/GlslBackend.cpp
--- a/Engine/Source/Developer/ShaderFormatOpenGL/Private/GlslBackend.cpp
+++ b/Engine/Source/Developer/ShaderFormatOpenGL/Private/GlslBackend.cpp
@@ -3823,16 +3823,24 @@ static void ConfigureInOutVariableLayout(EHlslShaderFrequency Frequency,
 	}
 	else if (Semantic && FCStringAnsi::Strnicmp(Variable->name, "gl_", 3) != 0)
 	{
-		Variable->explicit_location = 1;
 		Variable->semantic = ralloc_strdup(Variable, Semantic);
+
 		if(Mode == ir_var_in)
 		{
 			Variable->location = ParseState->next_in_location_slot++;
 		}
 		else
 		{
-			Variable->location = ParseState->next_out_location_slot++;
+			// Location may be already assigned to a pixel shader outputs (SV_TargetX HLSL output semantics).
+			// We want to preserve explicitly assigned render target indices and auto-generate them otherwise.
+			const bool bIsRenderTargetOutput = Frequency == HSF_PixelShader && Mode == ir_var_out;
+			if (!bIsRenderTargetOutput || !Variable->explicit_location)
+			{
+				Variable->location = ParseState->next_out_location_slot++;
+			}
 		}
+
+		Variable->explicit_location = true;
 	}
 }
 
diff --git a/Engine/Source/Developer/ShaderFormatOpenGL/Private/ShaderFormatOpenGL.cpp b/Engine/Source/Developer/ShaderFormatOpenGL/Private/ShaderFormatOpenGL.cpp
--- a/Engine/Source/Developer/ShaderFormatOpenGL/Private/ShaderFormatOpenGL.cpp
+++ b/Engine/Source/Developer/ShaderFormatOpenGL/Private/ShaderFormatOpenGL.cpp
@@ -24,16 +24,16 @@ class FShaderFormatGLSL : public IShaderFormat
 	enum
 	{
 		/** Version for shader format, this becomes part of the DDC key. */
-		UE_SHADER_GLSL_150_VER = 71,
-		UE_SHADER_GLSL_430_VER = 71,
-		UE_SHADER_GLSL_ES2_VER = 71,
-		UE_SHADER_GLSL_150ES2_VER = 71,
-		UE_SHADER_GLSL_150ES2NOUB_VER = 71,
-		UE_SHADER_GLSL_150ES3_1_VER = 71,
-		UE_SHADER_GLSL_ES2_VER_WEBGL = 71,
-		UE_SHADER_GLSL_ES2_IOS_VER = 71,
-		UE_SHADER_GLSL_310_ES_EXT_VER = 71,
-		UE_SHADER_GLSL_ES3_1_ANDROID_VER = 71,
+		UE_SHADER_GLSL_150_VER = 72,
+		UE_SHADER_GLSL_430_VER = 72,
+		UE_SHADER_GLSL_ES2_VER = 72,
+		UE_SHADER_GLSL_150ES2_VER = 72,
+		UE_SHADER_GLSL_150ES2NOUB_VER = 72,
+		UE_SHADER_GLSL_150ES3_1_VER = 72,
+		UE_SHADER_GLSL_ES2_VER_WEBGL = 72,
+		UE_SHADER_GLSL_ES2_IOS_VER = 72,
+		UE_SHADER_GLSL_310_ES_EXT_VER = 72,
+		UE_SHADER_GLSL_ES3_1_ANDROID_VER = 72,
 	}; 
 
 	void CheckFormat(FName Format) const
