Fixed: Some particles in ActionRPG are displayed as planes of various colors on some Android devices

diff --git a/Engine/Source/Developer/ShaderFormatOpenGL/Private/GlslBackend.cpp b/Engine/Source/Developer/ShaderFormatOpenGL/Private/GlslBackend.cpp
--- a/Engine/Source/Developer/ShaderFormatOpenGL/Private/GlslBackend.cpp
+++ b/Engine/Source/Developer/ShaderFormatOpenGL/Private/GlslBackend.cpp
@@ -3238,7 +3238,8 @@ bool compiler_internal_AdjustIsFrontFacing(bool isFrontFacing)
 				ralloc_asprintf_append(buffer, "	#define %sout\n", ES31FrameBufferFetchStorageQualifier);
 				ralloc_asprintf_append(buffer, "	vec4 FramebufferFetchES2() { return gl_LastFragColorARM; }\n");
 				ralloc_asprintf_append(buffer, "#else\n");
-				ralloc_asprintf_append(buffer, "	#error This shader requires framebuffer fetch support.\n");
+				ralloc_asprintf_append(buffer, "	#define %sout\n", ES31FrameBufferFetchStorageQualifier);
+				ralloc_asprintf_append(buffer, "	vec4 FramebufferFetchES2() { return vec4(65000.0, 65000.0, 65000.0, 65000.0); }\n");
 				ralloc_asprintf_append(buffer, "#endif\n\n");
 			}
 			else // ES3, ES2
diff --git a/Engine/Source/Developer/ShaderFormatOpenGL/Private/ShaderFormatOpenGL.cpp b/Engine/Source/Developer/ShaderFormatOpenGL/Private/ShaderFormatOpenGL.cpp
--- a/Engine/Source/Developer/ShaderFormatOpenGL/Private/ShaderFormatOpenGL.cpp
+++ b/Engine/Source/Developer/ShaderFormatOpenGL/Private/ShaderFormatOpenGL.cpp
@@ -24,16 +24,16 @@ class FShaderFormatGLSL : public IShaderFormat
 	enum
 	{
 		/** Version for shader format, this becomes part of the DDC key. */
-		UE_SHADER_GLSL_150_VER = 72,
-		UE_SHADER_GLSL_430_VER = 72,
-		UE_SHADER_GLSL_ES2_VER = 72,
-		UE_SHADER_GLSL_150ES2_VER = 72,
-		UE_SHADER_GLSL_150ES2NOUB_VER = 72,
-		UE_SHADER_GLSL_150ES3_1_VER = 72,
-		UE_SHADER_GLSL_ES2_VER_WEBGL = 72,
-		UE_SHADER_GLSL_ES2_IOS_VER = 72,
-		UE_SHADER_GLSL_310_ES_EXT_VER = 72,
-		UE_SHADER_GLSL_ES3_1_ANDROID_VER = 72,
+		UE_SHADER_GLSL_150_VER = 73,
+		UE_SHADER_GLSL_430_VER = 73,
+		UE_SHADER_GLSL_ES2_VER = 73,
+		UE_SHADER_GLSL_150ES2_VER = 73,
+		UE_SHADER_GLSL_150ES2NOUB_VER = 73,
+		UE_SHADER_GLSL_150ES3_1_VER = 73,
+		UE_SHADER_GLSL_ES2_VER_WEBGL = 73,
+		UE_SHADER_GLSL_ES2_IOS_VER = 73,
+		UE_SHADER_GLSL_310_ES_EXT_VER = 73,
+		UE_SHADER_GLSL_ES3_1_ANDROID_VER = 73,
 	}; 
 
 	void CheckFormat(FName Format) const
