Shaders compiled for GLSL_ES2 with CFLAG_FeatureLevelES31 compiler flag will now shift to GLSL_ES3_1_ANDROID profile instead of GLSL_310_ES_EXT.

diff --git a/Engine/Source/Developer/ShaderFormatOpenGL/Private/OpenGLShaderCompiler.cpp b/Engine/Source/Developer/ShaderFormatOpenGL/Private/OpenGLShaderCompiler.cpp
--- a/Engine/Source/Developer/ShaderFormatOpenGL/Private/OpenGLShaderCompiler.cpp
+++ b/Engine/Source/Developer/ShaderFormatOpenGL/Private/OpenGLShaderCompiler.cpp
@@ -1396,7 +1396,7 @@ void FOpenGLFrontend::CompileShader(const FShaderCompilerInput& Input,FShaderCom
 	const bool bCompileES2With310 = (Version == GLSL_ES2 && Input.Environment.CompilerFlags.Contains(CFLAG_FeatureLevelES31));
 	if (bCompileES2With310)
 	{
-		Version = GLSL_310_ES_EXT;
+		Version = GLSL_ES3_1_ANDROID;
 	}
 
 	// set up compiler env based on version
