[UE-53816] Allow texture sampling in vertex shaders for ES2 feature level

diff --git a/Engine/Source/Developer/ShaderFormatOpenGL/Private/GlslBackend.cpp b/Engine/Source/Developer/ShaderFormatOpenGL/Private/GlslBackend.cpp
--- a/Engine/Source/Developer/ShaderFormatOpenGL/Private/GlslBackend.cpp
+++ b/Engine/Source/Developer/ShaderFormatOpenGL/Private/GlslBackend.cpp
@@ -1400,7 +1400,7 @@ class ir_gen_glsl_visitor : public ir_visitor
 
 		bool bEmitEXT = false;
 
-		if (bIsES && op == ir_txl)
+		if (bIsES && op == ir_txl && ShaderTarget == fragment_shader)
 		{
 			// See http://www.khronos.org/registry/gles/extensions/EXT/EXT_shader_texture_lod.txt
 			bUsesES2TextureLODExtension = true;
@@ -3546,28 +3546,6 @@ struct SPromoteSampleLevelES2 : public ir_hierarchical_visitor
 
 	virtual ir_visitor_status visit_leave(ir_texture* IR) override
 	{
-		if (IR->op == ir_txl)
-		{
-			if (bIsVertexShader && bIsES2)
-			{
-				YYLTYPE loc;
-				loc.first_column = IR->SourceLocation.Column;
-				loc.first_line = IR->SourceLocation.Line;
-				loc.source_file = IR->SourceLocation.SourceFile;
-				_mesa_glsl_error(&loc, ParseState, "Vertex texture fetch currently not supported on GLSL ES\n");
-			}
-			else
-			{
-				//@todo-mobile: allowing lod texture functions for now, as they are supported on some devices via glsl extension.
-				// http://www.khronos.org/registry/gles/extensions/EXT/EXT_shader_texture_lod.txt
-				// Compat work will be required for devices which do not support it.
-				/*
-				_mesa_glsl_warning(ParseState, "%s(%u, %u) Converting SampleLevel() to Sample()\n", IR->SourceLocation.SourceFile.c_str(), IR->SourceLocation.Line, IR->SourceLocation.Column);
-				IR->op = ir_tex;
-				*/
-			}
-		}
-
 		if (IR->offset)
 		{
 			YYLTYPE loc;
diff --git a/Engine/Source/Runtime/Engine/Private/Materials/HLSLMaterialTranslator.h b/Engine/Source/Runtime/Engine/Private/Materials/HLSLMaterialTranslator.h
--- a/Engine/Source/Runtime/Engine/Private/Materials/HLSLMaterialTranslator.h
+++ b/Engine/Source/Runtime/Engine/Private/Materials/HLSLMaterialTranslator.h
@@ -3421,10 +3421,12 @@ protected:
 		}
 		else if(MipValueMode == TMVM_MipLevel)
 		{
+			// WebGL 2/GLES3.0 (or browsers with the texture lod extension) it is possible to sample from specific mip levels
+			// GLSL >= 100 should support this in the vertex shader
+			bool bES2MipSupport = (Platform == SP_OPENGL_ES2_WEBGL) || (Platform == SP_OPENGL_ES2_ANDROID && ShaderFrequency == SF_Vertex);
 			// Mobile: Sampling of a particular level depends on an extension; iOS does have it by default but
 			// there's a driver as of 7.0.2 that will cause a GPU hang if used with an Aniso > 1 sampler, so show an error for now
-			if ((Platform != SP_OPENGL_ES2_WEBGL) && // WebGL 2/GLES3.0 (or browsers with the texture lod extension) it is possible to sample from specific mip levels
-				ErrorUnlessFeatureLevelSupported(ERHIFeatureLevel::ES3_1) == INDEX_NONE)
+			if (!bES2MipSupport && ErrorUnlessFeatureLevelSupported(ERHIFeatureLevel::ES3_1) == INDEX_NONE)
 			{
 				Errorf(TEXT("Sampling for a specific mip-level is not supported for ES2"));
 				return INDEX_NONE;
