UE4.22 - Fix glsl dumped shader file paths
- Include cmd line for directcompile
#rnx

diff --git a/Engine/Source/Developer/ShaderFormatOpenGL/Private/OpenGLShaderCompiler.cpp b/Engine/Source/Developer/ShaderFormatOpenGL/Private/OpenGLShaderCompiler.cpp
--- a/Engine/Source/Developer/ShaderFormatOpenGL/Private/OpenGLShaderCompiler.cpp
+++ b/Engine/Source/Developer/ShaderFormatOpenGL/Private/OpenGLShaderCompiler.cpp
@@ -1494,6 +1494,12 @@ void FOpenGLFrontend::CompileShader(const FShaderCompilerInput& Input,FShaderCom
 	// This requires removing the HLSLCC_NoPreprocess flag later on!
 	RemoveUniformBuffersFromSource(Input.Environment, PreprocessedShader);
 
+	uint32 CCFlags = CalculateCrossCompilerFlags(Version, Input.Environment.CompilerFlags);
+
+	// Required as we added the RemoveUniformBuffersFromSource() function (the cross-compiler won't be able to interpret comments w/o a preprocessor)
+	CCFlags &= ~HLSLCC_NoPreprocess;
+
+
 	// Write out the preprocessed file and a batch file to compile it if requested (DumpDebugInfoPath is valid)
 	if (bDumpDebugInfo)
 	{
@@ -1504,6 +1510,11 @@ void FOpenGLFrontend::CompileShader(const FShaderCompilerInput& Input,FShaderCom
 			FileWriter->Serialize((ANSICHAR*)AnsiSourceFile.Get(), AnsiSourceFile.Length());
 			{
 				FString Line = CrossCompiler::CreateResourceTableFromEnvironment(Input.Environment);
+
+				Line += TEXT("#if 0 /*DIRECT COMPILE*/\n");
+				Line += CreateShaderCompilerWorkerDirectCommandLine(Input, CCFlags);
+				Line += TEXT("\n#endif /*DIRECT COMPILE*/\n");
+
 				FileWriter->Serialize(TCHAR_TO_ANSI(*Line), Line.Len());
 			}
 			FileWriter->Close();
@@ -1516,11 +1527,6 @@ void FOpenGLFrontend::CompileShader(const FShaderCompilerInput& Input,FShaderCom
 		}
 	}
 
-	uint32 CCFlags = CalculateCrossCompilerFlags(Version, Input.Environment.CompilerFlags);
-
-	// Required as we added the RemoveUniformBuffersFromSource() function (the cross-compiler won't be able to interpret comments w/o a preprocessor)
-	CCFlags &= ~HLSLCC_NoPreprocess;
-
 	FGlslCodeBackend* BackEnd = CreateBackend(Version, CCFlags, HlslCompilerTarget);
 	FGlslLanguageSpec* LanguageSpec = CreateLanguageSpec(Version);
 
@@ -1567,7 +1573,7 @@ void FOpenGLFrontend::CompileShader(const FShaderCompilerInput& Input,FShaderCom
 				GlslShaderSource = Dest;
 				GlslSourceLen = FCStringAnsi::Strlen(GlslShaderSource);
 
-				FArchive* FileWriter = IFileManager::Get().CreateFileWriter(*(Input.DumpDebugInfoPath / Input.VirtualSourceFilePath + TEXT(".glsl")));
+				FArchive* FileWriter = IFileManager::Get().CreateFileWriter(*GLSLFile);
 				if (FileWriter)
 				{
 					FileWriter->Serialize(GlslShaderSource,GlslSourceLen+1);
@@ -2005,7 +2011,9 @@ void FOpenGLFrontend::CompileOffline(const FShaderCompilerInput& Input, FShaderC
 	const bool bSupportsOfflineCompilation = PlatformSupportsOfflineCompilation(ShaderVersion);
 
 	if (!bSupportsOfflineCompilation)
-		return;	
+	{
+		return;
+	}
 
 	TSharedPtr<ANSICHAR> ShaderSource = PrepareCodeForOfflineCompilation(ShaderVersion, (EShaderFrequency)Input.Target.Frequency, InShaderSource);
 
@@ -2015,5 +2023,7 @@ void FOpenGLFrontend::CompileOffline(const FShaderCompilerInput& Input, FShaderC
 void FOpenGLFrontend::PlatformCompileOffline(const FShaderCompilerInput& Input, FShaderCompilerOutput& ShaderOutput, const ANSICHAR* ShaderSource, const GLSLVersion ShaderVersion)
 {
 	if (ShaderVersion == GLSL_ES2 || ShaderVersion == GLSL_ES3_1_ANDROID || ShaderVersion == GLSL_ES2_IOS)
+	{
 		CompileOfflineMali(Input, ShaderOutput, ShaderSource, FPlatformString::Strlen(ShaderSource), false);
+	}
 }
