Fix for ES31 shaders attempting to use framebuffer fetch in the vertex shader, causing a crash on Android

diff --git a/Engine/Shaders/Public/ShaderVersion.ush b/Engine/Shaders/Public/ShaderVersion.ush
--- a/Engine/Shaders/Public/ShaderVersion.ush
+++ b/Engine/Shaders/Public/ShaderVersion.ush
@@ -3,4 +3,4 @@
 // in Platform.ush (which should be included in any shader) it allows to invalidate the shader DDC.
 // 
 // If you are merging streams and there is a conflict with this GUID you should make a new GUID rather than taking one or the other.
-// GUID = C15DFE5B413C0C4E6C49949A06F0AE74
\ No newline at end of file
+// GUID = 0D98D1EF7B834118A251910A92A95B92
\ No newline at end of file
diff --git a/Engine/Source/Developer/ShaderFormatOpenGL/Private/GlslBackend.cpp b/Engine/Source/Developer/ShaderFormatOpenGL/Private/GlslBackend.cpp
--- a/Engine/Source/Developer/ShaderFormatOpenGL/Private/GlslBackend.cpp
+++ b/Engine/Source/Developer/ShaderFormatOpenGL/Private/GlslBackend.cpp
@@ -569,6 +569,9 @@ class ir_gen_glsl_visitor : public ir_visitor
 	// framebuffer fetch is in use
 	bool bUsesFrameBufferFetch;
 
+	// depthbuffer fetch is in use
+	bool bUsesDepthbufferFetch;
+
 	// uses external texture
 	bool bUsesExternalTexture;
 
@@ -3124,7 +3127,7 @@ class ir_gen_glsl_visitor : public ir_visitor
 public:
 
 	/** Constructor. */
-	ir_gen_glsl_visitor(bool bInIsES, bool bInEmitPrecision, bool bInIsWebGL, EHlslCompileTarget InCompileTarget, _mesa_glsl_parser_targets InShaderTarget, bool bInGenerateLayoutLocations, bool bInDefaultPrecisionIsHalf, bool bInNoGlobalUniforms, bool bInUsesFrameBufferFetch, bool bInUsesExternalTexture)
+	ir_gen_glsl_visitor(bool bInIsES, bool bInEmitPrecision, bool bInIsWebGL, EHlslCompileTarget InCompileTarget, _mesa_glsl_parser_targets InShaderTarget, bool bInGenerateLayoutLocations, bool bInDefaultPrecisionIsHalf, bool bInNoGlobalUniforms, bool bInUsesFrameBufferFetch, bool bInUsesDepthbufferFetch, bool bInUsesExternalTexture)
 		: early_depth_stencil(false)
 		, bIsES(bInIsES)
 		, bEmitPrecision(bInEmitPrecision)
@@ -3134,6 +3137,7 @@ public:
 		, bGenerateLayoutLocations(bInGenerateLayoutLocations)
 		, bDefaultPrecisionIsHalf(bInDefaultPrecisionIsHalf)
 		, bUsesFrameBufferFetch(bInUsesFrameBufferFetch)
+		, bUsesDepthbufferFetch(bInUsesDepthbufferFetch)
 		, bUsesExternalTexture(bInUsesExternalTexture)
 		, buffer(0)
 		, indentation(0)
@@ -3275,8 +3279,7 @@ bool compiler_internal_AdjustIsFrontFacing(bool isFrontFacing)
 
 		}
 
-		bool bUsesDepthbufferFetchES2 = UsesUEIntrinsic(ir, DEPTHBUFFER_FETCH_ES2);
-		if (bUsesDepthbufferFetchES2)
+		if (bUsesDepthbufferFetch)
 		{
 			ralloc_asprintf_append(buffer, "\n#ifdef GL_ARM_shader_framebuffer_fetch_depth_stencil\n");
 			ralloc_asprintf_append(buffer, "float DepthbufferFetchES2(float OptionalDepth, float C1, float C2) { float w = 1.0/(gl_LastFragDepthARM*C1-C2); return clamp(w, 0.0, 65000.0); }\n");
@@ -3338,7 +3341,7 @@ bool compiler_internal_AdjustIsFrontFacing(bool isFrontFacing)
 
 		char* Extensions = ralloc_asprintf(mem_ctx, "");
 		buffer = &Extensions;
-		print_extensions(state, bUsesFrameBufferFetch, bUsesDepthbufferFetchES2, CompileTarget == HCT_FeatureLevelES3_1Ext, bUsesExternalTexture);
+		print_extensions(state, bUsesFrameBufferFetch, bUsesDepthbufferFetch, CompileTarget == HCT_FeatureLevelES3_1Ext, bUsesExternalTexture);
 		if (state->bSeparateShaderObjects && !state->bGenerateES)
 		{
 			switch (state->target)
@@ -3531,6 +3534,7 @@ char* FGlslCodeBackend::GenerateCode(exec_list* ir, _mesa_glsl_parse_state* stat
 	const bool bGenerateLayoutLocations = state->bGenerateLayoutLocations;
 	const bool bEmitPrecision = WantsPrecisionModifiers();
 	const bool bUsesFrameBufferFetch = Frequency == HSF_PixelShader && UsesUEIntrinsic(ir, FRAMEBUFFER_FETCH_ES2);
+	const bool bUsesDepthBufferFetch = Frequency == HSF_PixelShader && UsesUEIntrinsic(ir, DEPTHBUFFER_FETCH_ES2);
 	ir_gen_glsl_visitor visitor(state->bGenerateES,
 								bEmitPrecision,
 								bIsWebGL,
@@ -3540,6 +3544,7 @@ char* FGlslCodeBackend::GenerateCode(exec_list* ir, _mesa_glsl_parse_state* stat
 								bDefaultPrecisionIsHalf,
 								!AllowsGlobalUniforms(),
 								bUsesFrameBufferFetch,
+								bUsesDepthBufferFetch,
 								bUsesExternalTexture);
 	const char* code = visitor.run(ir, state, bGroupFlattenedUBs);
 	return _strdup(code);
diff --git a/Engine/Source/Developer/ShaderFormatOpenGL/Private/ShaderFormatOpenGL.cpp b/Engine/Source/Developer/ShaderFormatOpenGL/Private/ShaderFormatOpenGL.cpp
--- a/Engine/Source/Developer/ShaderFormatOpenGL/Private/ShaderFormatOpenGL.cpp
+++ b/Engine/Source/Developer/ShaderFormatOpenGL/Private/ShaderFormatOpenGL.cpp
@@ -24,8 +24,8 @@ class FShaderFormatGLSL : public IShaderFormat
 	enum
 	{
 		/** Version for shader format, this becomes part of the DDC key. */
-		UE_SHADER_GLSL_VER = 79,
-		UE_SHADER_GLSL_ANDROID_VER = 79,
+		UE_SHADER_GLSL_VER = 80,
+		UE_SHADER_GLSL_ANDROID_VER = 80,
 	};
 
 	void CheckFormat(FName Format) const
diff --git a/Engine/Source/Runtime/Engine/Private/Materials/HLSLMaterialTranslator.h b/Engine/Source/Runtime/Engine/Private/Materials/HLSLMaterialTranslator.h
--- a/Engine/Source/Runtime/Engine/Private/Materials/HLSLMaterialTranslator.h
+++ b/Engine/Source/Runtime/Engine/Private/Materials/HLSLMaterialTranslator.h
@@ -3728,6 +3728,13 @@ protected:
 
 	virtual int32 SceneDepth(int32 Offset, int32 ViewportUV, bool bUseOffset) override
 	{
+		if (ShaderFrequency == SF_Vertex && FeatureLevel <= ERHIFeatureLevel::ES3_1)
+		{
+			// mobile currently does not support this, we need to read a separate copy of the depth, we must disable framebuffer fetch and force scene texture reads.
+			//return Errorf(TEXT("Cannot read scene depth from the vertex shader with feature level ES3.1 or below."));
+			return Constant(65000.0f);
+		}
+
 		if (Offset == INDEX_NONE && bUseOffset)
 		{
 			return INDEX_NONE;
