DR - Workaround GL generating bad packed uniforms when OpenGL.EmulatedUBs = 1
#rnx

diff --git a/Engine/Source/Developer/ShaderFormatOpenGL/Private/OpenGLShaderCompiler.cpp b/Engine/Source/Developer/ShaderFormatOpenGL/Private/OpenGLShaderCompiler.cpp
--- a/Engine/Source/Developer/ShaderFormatOpenGL/Private/OpenGLShaderCompiler.cpp
+++ b/Engine/Source/Developer/ShaderFormatOpenGL/Private/OpenGLShaderCompiler.cpp
@@ -1372,7 +1372,8 @@ uint32 FOpenGLFrontend::CalculateCrossCompilerFlags(GLSLVersion Version, const T
 		CompilerFlags.Contains(CFLAG_UseEmulatedUB))
 	{
 		CCFlags |= HLSLCC_FlattenUniformBuffers | HLSLCC_FlattenUniformBufferStructures;
-		CCFlags |= HLSLCC_GroupFlattenedUniformBuffers;
+		// Do not use this flag as it adds too many uniforms.
+		//CCFlags |= HLSLCC_GroupFlattenedUniformBuffers;
 		CCFlags |= HLSLCC_ExpandUBMemberArrays;
 	}
 
diff --git a/Engine/Source/Developer/ShaderFormatOpenGL/Private/ShaderFormatOpenGL.cpp b/Engine/Source/Developer/ShaderFormatOpenGL/Private/ShaderFormatOpenGL.cpp
--- a/Engine/Source/Developer/ShaderFormatOpenGL/Private/ShaderFormatOpenGL.cpp
+++ b/Engine/Source/Developer/ShaderFormatOpenGL/Private/ShaderFormatOpenGL.cpp
@@ -24,17 +24,9 @@ class FShaderFormatGLSL : public IShaderFormat
 	enum
 	{
 		/** Version for shader format, this becomes part of the DDC key. */
-		UE_SHADER_GLSL_150_VER = 77,
-		UE_SHADER_GLSL_430_VER = 77,
-		UE_SHADER_GLSL_ES2_VER = 77,
-		UE_SHADER_GLSL_150ES2_VER = 77,
-		UE_SHADER_GLSL_150ES2NOUB_VER = 77,
-		UE_SHADER_GLSL_150ES3_1_VER = 77,
-		UE_SHADER_GLSL_ES2_VER_WEBGL = 77,
-		UE_SHADER_GLSL_ES2_IOS_VER = 77,
-		UE_SHADER_GLSL_310_ES_EXT_VER = 77,
-		UE_SHADER_GLSL_ES3_1_ANDROID_VER = 77,
-	}; 
+		UE_SHADER_GLSL_VER = 78,
+		UE_SHADER_GLSL_ANDROID_VER = 78,
+	};
 
 	void CheckFormat(FName Format) const
 	{
@@ -58,43 +50,43 @@ public:
 		uint32 GLSLVersion = 0;
 		if (Format == NAME_GLSL_150)
 		{
-			GLSLVersion  = UE_SHADER_GLSL_150_VER;
+			GLSLVersion  = UE_SHADER_GLSL_VER;
 		}
 		else if (Format == NAME_GLSL_430)
 		{
-			GLSLVersion = UE_SHADER_GLSL_430_VER;
+			GLSLVersion = UE_SHADER_GLSL_VER;
 		}
 		else if (Format == NAME_GLSL_ES2)
 		{
-			GLSLVersion = UE_SHADER_GLSL_ES2_VER;
+			GLSLVersion = UE_SHADER_GLSL_ANDROID_VER;
 		}
 		else if (Format == NAME_GLSL_150_ES2)
 		{
-			GLSLVersion = UE_SHADER_GLSL_150ES2_VER;
+			GLSLVersion = UE_SHADER_GLSL_ANDROID_VER;
 		}
 		else if (Format == NAME_GLSL_150_ES3_1)
 		{
-			GLSLVersion = UE_SHADER_GLSL_150ES3_1_VER;
+			GLSLVersion = UE_SHADER_GLSL_ANDROID_VER;
 		}
 		else if (Format == NAME_GLSL_150_ES2_NOUB)
 		{
-			GLSLVersion = UE_SHADER_GLSL_150ES2NOUB_VER;
+			GLSLVersion = UE_SHADER_GLSL_ANDROID_VER;
 		}
 		else if (Format == NAME_GLSL_ES2_WEBGL)
 		{
-			GLSLVersion = UE_SHADER_GLSL_ES2_VER_WEBGL;
+			GLSLVersion = UE_SHADER_GLSL_ANDROID_VER;
 		}
 		else if (Format == NAME_GLSL_ES2_IOS)
 		{
-			GLSLVersion = UE_SHADER_GLSL_ES2_IOS_VER;
+			GLSLVersion = UE_SHADER_GLSL_ANDROID_VER;
 		}
 		else if (Format == NAME_GLSL_310_ES_EXT)
 		{
-			GLSLVersion = UE_SHADER_GLSL_310_ES_EXT_VER;
+			GLSLVersion = UE_SHADER_GLSL_ANDROID_VER;
 		}
 		else if (Format == NAME_GLSL_ES3_1_ANDROID)
 		{
-			GLSLVersion = UE_SHADER_GLSL_ES3_1_ANDROID_VER;
+			GLSLVersion = UE_SHADER_GLSL_ANDROID_VER;
 		}
 		else
 		{
