UE4.22 - Fix invalid RWBuffer translation on OpenGL
#rnx

diff --git a/Engine/Source/Developer/ShaderFormatOpenGL/Private/GlslBackend.cpp b/Engine/Source/Developer/ShaderFormatOpenGL/Private/GlslBackend.cpp
--- a/Engine/Source/Developer/ShaderFormatOpenGL/Private/GlslBackend.cpp
+++ b/Engine/Source/Developer/ShaderFormatOpenGL/Private/GlslBackend.cpp
@@ -1659,15 +1659,15 @@ class ir_gen_glsl_visitor : public ir_visitor
 		{
 			if (bIsStructured)
 			{
-				if (src)
-				{
-					src->accept(this);
-					ralloc_asprintf_append(buffer, " = ");
-				}
 				deref->image->accept(this);
 				ralloc_asprintf_append(buffer, "[");
 				deref->image_index->accept(this);
 				ralloc_asprintf_append(buffer, "]");
+				if (src)
+				{
+					ralloc_asprintf_append(buffer, " = ");
+					src->accept(this);
+				}
 			}
 			else
 			{
diff --git a/Engine/Source/Developer/ShaderFormatOpenGL/Private/ShaderFormatOpenGL.cpp b/Engine/Source/Developer/ShaderFormatOpenGL/Private/ShaderFormatOpenGL.cpp
--- a/Engine/Source/Developer/ShaderFormatOpenGL/Private/ShaderFormatOpenGL.cpp
+++ b/Engine/Source/Developer/ShaderFormatOpenGL/Private/ShaderFormatOpenGL.cpp
@@ -24,7 +24,7 @@ class FShaderFormatGLSL : public IShaderFormat
 	enum
 	{
 		/** Version for shader format, this becomes part of the DDC key. */
-		UE_SHADER_GLSL_VER = 78,
+		UE_SHADER_GLSL_VER = 79,
 		UE_SHADER_GLSL_ANDROID_VER = 78,
 	};
 
