Fixed cross-compilation of patch constant data for GLSL backend on AMD GPUs.

diff --git a/Engine/Source/Developer/ShaderFormatOpenGL/Private/GlslBackend.cpp b/Engine/Source/Developer/ShaderFormatOpenGL/Private/GlslBackend.cpp
--- a/Engine/Source/Developer/ShaderFormatOpenGL/Private/GlslBackend.cpp
+++ b/Engine/Source/Developer/ShaderFormatOpenGL/Private/GlslBackend.cpp
@@ -980,7 +980,54 @@ class ir_gen_glsl_visitor : public ir_visitor
 				bUseGlobalUniformBufferWrapper = true;
 			}
 
-			if (scope_depth == 0 &&
+			if (var->is_patch_constant)
+			{
+				// AMD drivers reject interface blocks for per-patch data.
+				// AMD drivers also need a location qualifier for each shader input/output vector.
+				// So we translate patch constant data to individual structs:
+				//   "layout(location = 9) patch in struct { vec4 Data; } in_PN_POSITION9;"
+				// NVIDIA drivers would also accept the previous solution:
+				//   "patch in PN_POSITION9 { vec4 Data; } in_PN_POSITION9;"
+
+				// We expect a struct with single member "Data" at this point
+				check(var->type->base_type == GLSL_TYPE_STRUCT);
+
+				// Patch declarations cannot have interpolation qualifiers
+				if (var->explicit_location)
+				{
+					ralloc_asprintf_append(
+						buffer,
+						"layout(location = %d) patch %sstruct",
+						var->location,
+						mode_str[var->mode]
+					);
+				}
+				else
+				{
+					ralloc_asprintf_append(
+						buffer,
+						"patch %sstruct",
+						mode_str[var->mode]
+					);
+				}
+
+				const glsl_type* inner_type = var->type;
+				if (inner_type->is_array())
+				{
+					inner_type = inner_type->fields.array;
+				}
+				check(inner_type->is_record());
+				check(inner_type->length == 1);
+				const glsl_struct_field* field = &inner_type->fields.structure[0];
+				check(strcmp(field->name, "Data") == 0);
+
+				ralloc_asprintf_append(buffer, " { ");
+				print_type_pre(field->type);
+				ralloc_asprintf_append(buffer, " Data");
+				print_type_post(field->type);
+				ralloc_asprintf_append(buffer, "; }");
+			}
+			else if (scope_depth == 0 &&
 			   ((var->mode == ir_var_in) || (var->mode == ir_var_out)) && 
 			   var->is_interface_block)
 			{
@@ -1014,7 +1061,7 @@ class ir_gen_glsl_visitor : public ir_visitor
 					{
 						interp_qualifier = "smooth ";
 					}
-										
+
 					ralloc_asprintf_append(
 						buffer,
 						"INTERFACE_BLOCK(%d, %s, %s%s%s%s, ",
@@ -4121,7 +4168,7 @@ static ir_rvalue* GenShaderInputSemantic(
 		Variable->interpolation = InputQualifier.Fields.InterpolationMode;
 		Variable->is_patch_constant = InputQualifier.Fields.bIsPatchConstant;
 
-		if(ParseState->bGenerateLayoutLocations && !InputQualifier.Fields.bIsPatchConstant)
+		if(ParseState->bGenerateLayoutLocations)
 		{
 			ConfigureInOutVariableLayout(Frequency, ParseState, Semantic, Variable, ir_var_in);
 		}
@@ -4158,7 +4205,7 @@ static ir_rvalue* GenShaderInputSemantic(
 		DeclInstructions->push_tail(Variable);
 		ParseState->symbols->add_variable(Variable);
 
-		if (ParseState->bGenerateLayoutLocations && !Variable->is_patch_constant)
+		if (ParseState->bGenerateLayoutLocations)
 		{
 			ConfigureInOutVariableLayout(Frequency, ParseState, Semantic, Variable, ir_var_in);
 		}
@@ -4276,9 +4323,9 @@ static ir_rvalue* GenShaderOutputSemantic(
 			{
 				Variable->explicit_location = true;
 				Variable->location = OutputIndex;
+			}
 		}
 	}
-	}
 
 	if (Variable == NULL && Frequency == HSF_HullShader)
 	{
@@ -4335,7 +4382,7 @@ static ir_rvalue* GenShaderOutputSemantic(
 		Variable = new(ParseState)ir_variable(Type, ralloc_asprintf(ParseState, "var_%s", Semantic), ir_var_out);
 	}
 
-	if (ParseState->bGenerateLayoutLocations && Variable && !Variable->is_patch_constant)
+	if (ParseState->bGenerateLayoutLocations && Variable)
 	{
 		ConfigureInOutVariableLayout(Frequency, ParseState, Semantic, Variable, ir_var_out);
 	}
@@ -4381,7 +4428,7 @@ static ir_rvalue* GenShaderOutputSemantic(
 	Variable->is_interface_block = true;
 	Variable->is_patch_constant = OutputQualifier.Fields.bIsPatchConstant;
 
-	if (ParseState->bGenerateLayoutLocations && !Variable->is_patch_constant)
+	if (ParseState->bGenerateLayoutLocations)
 	{
 		ConfigureInOutVariableLayout(Frequency, ParseState, Semantic, Variable, ir_var_out);
 	}
diff --git a/Engine/Source/Developer/ShaderFormatOpenGL/Private/ShaderFormatOpenGL.cpp b/Engine/Source/Developer/ShaderFormatOpenGL/Private/ShaderFormatOpenGL.cpp
--- a/Engine/Source/Developer/ShaderFormatOpenGL/Private/ShaderFormatOpenGL.cpp
+++ b/Engine/Source/Developer/ShaderFormatOpenGL/Private/ShaderFormatOpenGL.cpp
@@ -24,7 +24,7 @@ class FShaderFormatGLSL : public IShaderFormat
 	enum
 	{
 		/** Version for shader format, this becomes part of the DDC key. */
-		UE_SHADER_GLSL_VER = 81,
+		UE_SHADER_GLSL_VER = 82,
 		UE_SHADER_GLSL_ANDROID_VER = 81,
 	};
 
