Make OpenGLShaderCompiler dlsym GL funcs instead of relying on dangling syms.

diff --git a/Engine/Source/Developer/ShaderFormatOpenGL/Private/OpenGLShaderCompiler.cpp b/Engine/Source/Developer/ShaderFormatOpenGL/Private/OpenGLShaderCompiler.cpp
--- a/Engine/Source/Developer/ShaderFormatOpenGL/Private/OpenGLShaderCompiler.cpp
+++ b/Engine/Source/Developer/ShaderFormatOpenGL/Private/OpenGLShaderCompiler.cpp
@@ -29,12 +29,9 @@
 	#include <GL/wglext.h>
 #include "Windows/HideWindowsPlatformTypes.h"
 #elif PLATFORM_LINUX
-	#define GL_GLEXT_PROTOTYPES 1
 	#include <GL/glcorearb.h>
 	#include <GL/glext.h>
 	#include "SDL.h"
-	GLAPI GLuint APIENTRY glCreateShader (GLenum type);
-	GLAPI void APIENTRY glShaderSource (GLuint shader, GLsizei count, const GLchar* const *string, const GLint *length);
 	typedef SDL_Window*		SDL_HWindow;
 	typedef SDL_GLContext	SDL_HGLContext;
 	struct FPlatformOpenGLContext
@@ -281,6 +278,27 @@ static void PlatformReleaseOpenGL(void* ContextPtr, void* PrevContextPtr)
 	wglMakeCurrent((HDC)ContextPtr, (HGLRC)PrevContextPtr);
 }
 #elif PLATFORM_LINUX
+/** List all OpenGL entry points needed for shader compilation. */
+#define ENUM_GL_ENTRYPOINTS(EnumMacro) \
+	EnumMacro(PFNGLCOMPILESHADERPROC,glCompileShader) \
+	EnumMacro(PFNGLCREATESHADERPROC,glCreateShader) \
+	EnumMacro(PFNGLDELETESHADERPROC,glDeleteShader) \
+	EnumMacro(PFNGLGETSHADERIVPROC,glGetShaderiv) \
+	EnumMacro(PFNGLGETSHADERINFOLOGPROC,glGetShaderInfoLog) \
+	EnumMacro(PFNGLSHADERSOURCEPROC,glShaderSource) \
+	EnumMacro(PFNGLDELETEBUFFERSPROC,glDeleteBuffers)
+
+/** Define all GL functions. */
+// We need to make pointer names different from GL functions otherwise we may end up getting
+// addresses of those symbols when looking for extensions.
+namespace GLFuncPointers
+{
+	#define DEFINE_GL_ENTRYPOINTS(Type,Func) static Type Func = NULL;
+	ENUM_GL_ENTRYPOINTS(DEFINE_GL_ENTRYPOINTS);
+};
+
+using namespace GLFuncPointers;
+
 static void _PlatformCreateDummyGLWindow(FPlatformOpenGLContext *OutContext)
 {
 	static bool bInitializedWindowClass = false;
@@ -332,6 +350,19 @@ static void PlatformInitOpenGL(void*& ContextPtr, void*& PrevContextPtr, int InM
 			UE_LOG(LogOpenGLShaderCompiler, Fatal, TEXT("Unable to dynamically load libGL: %s"), ANSI_TO_TCHAR(SDL_GetError()));
 		}
 
+		if (glCreateShader == nullptr)
+		{
+			// Initialize all entry points.
+			#define GET_GL_ENTRYPOINTS(Type,Func) GLFuncPointers::Func = reinterpret_cast<Type>(SDL_GL_GetProcAddress(#Func));
+			ENUM_GL_ENTRYPOINTS(GET_GL_ENTRYPOINTS);
+
+			// Check that all of the entry points have been initialized.
+			bool bFoundAllEntryPoints = true;
+			#define CHECK_GL_ENTRYPOINTS(Type,Func) if (Func == nullptr) { bFoundAllEntryPoints = false; UE_LOG(LogOpenGLShaderCompiler, Warning, TEXT("Failed to find entry point for %s"), TEXT(#Func)); }
+			ENUM_GL_ENTRYPOINTS(CHECK_GL_ENTRYPOINTS);
+			checkf(bFoundAllEntryPoints, TEXT("Failed to find all OpenGL entry points."));
+		}
+
 		if	(SDL_GL_SetAttribute( SDL_GL_CONTEXT_MAJOR_VERSION, InMajorVersion))
 		{
 			UE_LOG(LogOpenGLShaderCompiler, Fatal, TEXT("Failed to set GL major version: %s"), ANSI_TO_TCHAR(SDL_GetError()));
