Fix same shader generation bug in GlslBackend as in VulkanBackend (regarding SV_ClipDistance).

diff --git a/Engine/Source/Developer/ShaderFormatOpenGL/Private/GlslBackend.cpp b/Engine/Source/Developer/ShaderFormatOpenGL/Private/GlslBackend.cpp
--- a/Engine/Source/Developer/ShaderFormatOpenGL/Private/GlslBackend.cpp
+++ b/Engine/Source/Developer/ShaderFormatOpenGL/Private/GlslBackend.cpp
@@ -4213,7 +4213,7 @@ static ir_rvalue* GenShaderOutputSemantic(
 		}
 	}
 
-	if (Variable == NULL && Frequency == HSF_VertexShader)
+	if (Variable == NULL && (Frequency == HSF_VertexShader || Frequency == HSF_GeometryShader || Frequency == HSF_HullShader || Frequency == HSF_DomainShader))
 	{
 		const int PrefixLength = 15;
 		if (FCStringAnsi::Strnicmp(Semantic, "SV_ClipDistance", PrefixLength) == 0)
diff --git a/Engine/Source/Developer/ShaderFormatOpenGL/Private/ShaderFormatOpenGL.cpp b/Engine/Source/Developer/ShaderFormatOpenGL/Private/ShaderFormatOpenGL.cpp
--- a/Engine/Source/Developer/ShaderFormatOpenGL/Private/ShaderFormatOpenGL.cpp
+++ b/Engine/Source/Developer/ShaderFormatOpenGL/Private/ShaderFormatOpenGL.cpp
@@ -24,8 +24,8 @@ class FShaderFormatGLSL : public IShaderFormat
 	enum
 	{
 		/** Version for shader format, this becomes part of the DDC key. */
-		UE_SHADER_GLSL_VER = 80,
-		UE_SHADER_GLSL_ANDROID_VER = 80,
+		UE_SHADER_GLSL_VER = 81,
+		UE_SHADER_GLSL_ANDROID_VER = 81,
 	};
 
 	void CheckFormat(FName Format) const
