Niagara GPUSim Ios/Android

glGenVertexArrays, glBindVertexArray, glMapBufferRange, glCopyBufferSubData, glDrawArraysIndirect, glDrawElementsIndirect added
One VAO bound per context now(different from Default VAO - indirect ES requirement)

diff --git a/Engine/Plugins/FX/Niagara/Source/Niagara/Classes/NiagaraEmitterInstance.h b/Engine/Plugins/FX/Niagara/Source/Niagara/Classes/NiagaraEmitterInstance.h
--- a/Engine/Plugins/FX/Niagara/Source/Niagara/Classes/NiagaraEmitterInstance.h
+++ b/Engine/Plugins/FX/Niagara/Source/Niagara/Classes/NiagaraEmitterInstance.h
@@ -68,6 +68,7 @@ public:
 	FNiagaraDataSet& GetData()const { return *ParticleDataSet; }
 
 	FORCEINLINE bool IsDisabled()const { return ExecutionState == ENiagaraExecutionState::Disabled; }
+	FORCEINLINE bool IsInactive()const { return ExecutionState == ENiagaraExecutionState::Inactive; }
 	FORCEINLINE bool IsComplete()const { return ExecutionState == ENiagaraExecutionState::Complete || ExecutionState == ENiagaraExecutionState::Disabled; }
 
 	/** Create a new NiagaraRenderer. The old renderer is not immediately deleted, but instead put in the ToBeRemoved list.*/
diff --git a/Engine/Plugins/FX/Niagara/Source/Niagara/Private/NiagaraComponent.cpp b/Engine/Plugins/FX/Niagara/Source/Niagara/Private/NiagaraComponent.cpp
--- a/Engine/Plugins/FX/Niagara/Source/Niagara/Private/NiagaraComponent.cpp
+++ b/Engine/Plugins/FX/Niagara/Source/Niagara/Private/NiagaraComponent.cpp
@@ -323,7 +323,7 @@ void FNiagaraSceneProxy::GetDynamicMeshElements(const TArray<const FSceneView*>&
 	for (int32 RendererIdx : RendererDrawOrder)
 	{
 		FNiagaraRenderer* Renderer = EmitterRenderers[RendererIdx];
-		if (Renderer && (Renderer->GetSimTarget() == ENiagaraSimTarget::CPUSim || ViewFamily.GetFeatureLevel() == ERHIFeatureLevel::SM5))
+		if (Renderer && (Renderer->GetSimTarget() == ENiagaraSimTarget::CPUSim || ViewFamily.GetFeatureLevel() >= ERHIFeatureLevel::ES3_1))
 		{
 			Renderer->GetDynamicMeshElements(Views, ViewFamily, VisibilityMap, Collector, this);
 		}
diff --git a/Engine/Plugins/FX/Niagara/Source/Niagara/Public/NiagaraCommon.h b/Engine/Plugins/FX/Niagara/Source/Niagara/Public/NiagaraCommon.h
--- a/Engine/Plugins/FX/Niagara/Source/Niagara/Public/NiagaraCommon.h
+++ b/Engine/Plugins/FX/Niagara/Source/Niagara/Public/NiagaraCommon.h
@@ -122,8 +122,7 @@ enum class ENiagaraScriptCompileStatus : uint8
 
 FORCEINLINE bool NiagaraSupportsComputeShaders(EShaderPlatform ShaderPlatform)
 {
-	// Change to RHISupportsComputeShaders(ShaderPlatform) to support ES3_1
-	return IsFeatureLevelSupported(ShaderPlatform, ERHIFeatureLevel::SM5);
+	return RHISupportsComputeShaders(ShaderPlatform);
 }
 
 USTRUCT()
diff --git a/Engine/Source/Developer/ShaderFormatOpenGL/Private/GlslBackend.cpp b/Engine/Source/Developer/ShaderFormatOpenGL/Private/GlslBackend.cpp
--- a/Engine/Source/Developer/ShaderFormatOpenGL/Private/GlslBackend.cpp
+++ b/Engine/Source/Developer/ShaderFormatOpenGL/Private/GlslBackend.cpp
@@ -611,7 +611,10 @@ class ir_gen_glsl_visitor : public ir_visitor
 	bool bUsesES2TextureLODExtension;
 
 	/** Whether the shader being cross compiled needs GL_EXT_texture_buffer. */
-	bool bUsesTexelFetch;
+	bool bUsesTextureBuffer;
+
+	/** Whether the shader being cross compiled needs GL_OES_shader_image_atomic. */
+	bool bUseImageAtomic;
 
 	// Found dFdx or dFdy
 	bool bUsesDXDY;
@@ -1480,7 +1483,7 @@ class ir_gen_glsl_visitor : public ir_visitor
 
 		if (op == ir_txf)
 		{
-			bUsesTexelFetch = true;
+			bUsesTextureBuffer = true;
 		}
 
 		// Emit texture function and sampler.
@@ -1741,6 +1744,7 @@ class ir_gen_glsl_visitor : public ir_visitor
 			}
 			else
 			{
+				bUsesTextureBuffer = true;
 				if (src == NULL)
 				{
 					ralloc_asprintf_append(buffer, "imageLoad( ");
@@ -2291,6 +2295,7 @@ class ir_gen_glsl_visitor : public ir_visitor
 		ir_dereference_image* image = ir->memory_ref->as_dereference_image();
 
 		ir->lhs->accept(this);
+		bUseImageAtomic = image != NULL;
 		if (!image || (image->image->type && image->image->type->shader_storage_buffer))
 		{
 			ralloc_asprintf_append(buffer, " = %s(",
@@ -3156,6 +3161,10 @@ class ir_gen_glsl_visitor : public ir_visitor
 			ralloc_asprintf_append(buffer, "#extension GL_EXT_gpu_shader5 : enable\n");
 			ralloc_asprintf_append(buffer, "\n#endif\n");
 			
+			ralloc_asprintf_append(buffer, "\n#ifdef GL_OES_shader_image_atomic\n");
+			ralloc_asprintf_append(buffer, "#extension GL_OES_shader_image_atomic : enable\n");
+			ralloc_asprintf_append(buffer, "\n#endif\n");
+
 			ralloc_asprintf_append(buffer, "\n#ifdef GL_EXT_texture_buffer\n");
 			ralloc_asprintf_append(buffer, "#extension GL_EXT_texture_buffer : enable\n");
 			ralloc_asprintf_append(buffer, "\n#endif\n");
@@ -3178,14 +3187,24 @@ class ir_gen_glsl_visitor : public ir_visitor
 				ralloc_asprintf_append(buffer, "#extension GL_EXT_tessellation_shader : enable\n");
 			}
 		}
-		else if ((bIsES || bIsES31) && bUsesTexelFetch)
+		else if ((bIsES || bIsES31))
 		{
-			// Not supported by ES2 and ES3.1 spec, but many phones support this extension
-			// GPU particles require this
-			// App shall not use a shader if this extension is not supported on device
-			ralloc_asprintf_append(buffer, "\n#ifdef GL_EXT_texture_buffer\n");
-			ralloc_asprintf_append(buffer, "#extension GL_EXT_texture_buffer : enable\n");
-			ralloc_asprintf_append(buffer, "\n#endif\n");
+			if (bUseImageAtomic)
+			{
+				ralloc_asprintf_append(buffer, "\n#ifdef GL_OES_shader_image_atomic\n");
+				ralloc_asprintf_append(buffer, "#extension GL_OES_shader_image_atomic : enable\n");
+				ralloc_asprintf_append(buffer, "\n#endif\n");
+			}
+
+			if (bUsesTextureBuffer)
+			{
+				// Not supported by ES2 and ES3.1 spec, but many phones support this extension
+				// GPU particles require this
+				// App shall not use a shader if this extension is not supported on device
+				ralloc_asprintf_append(buffer, "\n#ifdef GL_EXT_texture_buffer\n");
+				ralloc_asprintf_append(buffer, "#extension GL_EXT_texture_buffer : enable\n");
+				ralloc_asprintf_append(buffer, "\n#endif\n");
+			}
 		}
 		ralloc_asprintf_append(buffer, "// end extensions\n");
 	}
@@ -3214,7 +3233,8 @@ public:
 		, should_print_uint_literals_as_ints(false)
 		, loop_count(0)
 		, bUsesES2TextureLODExtension(false)
-		, bUsesTexelFetch(false)
+		, bUsesTextureBuffer(false)
+		, bUseImageAtomic(false)
 		, bUsesDXDY(false)
 		, bUsesInstanceID(false)
 		, bNoGlobalUniforms(bInNoGlobalUniforms)
diff --git a/Engine/Source/Developer/ShaderFormatOpenGL/Private/ShaderFormatOpenGL.cpp b/Engine/Source/Developer/ShaderFormatOpenGL/Private/ShaderFormatOpenGL.cpp
--- a/Engine/Source/Developer/ShaderFormatOpenGL/Private/ShaderFormatOpenGL.cpp
+++ b/Engine/Source/Developer/ShaderFormatOpenGL/Private/ShaderFormatOpenGL.cpp
@@ -23,8 +23,8 @@ class FShaderFormatGLSL : public IShaderFormat
 	enum
 	{
 		/** Version for shader format, this becomes part of the DDC key. */
-		UE_SHADER_GLSL_VER = 82,
-		UE_SHADER_GLSL_ANDROID_VER = 82,
+		UE_SHADER_GLSL_VER = 83,
+		UE_SHADER_GLSL_ANDROID_VER = 83,
 	};
 
 	void CheckFormat(FName Format) const
diff --git a/Engine/Source/Runtime/Core/Public/Android/AndroidPlatform.h b/Engine/Source/Runtime/Core/Public/Android/AndroidPlatform.h
--- a/Engine/Source/Runtime/Core/Public/Android/AndroidPlatform.h
+++ b/Engine/Source/Runtime/Core/Public/Android/AndroidPlatform.h
@@ -101,9 +101,6 @@ typedef FAndroidTypes FPlatformTypes;
 #ifndef USE_ANDROID_EVENTS
 	#define USE_ANDROID_EVENTS						1
 #endif
-#ifndef USE_ANDROID_OPENGL
-	#define USE_ANDROID_OPENGL						1
-#endif
 
 // Function type macros.
 #define VARARGS													/* Functions with variable arguments */
diff --git a/Engine/Source/Runtime/Engine/Private/Particles/ParticleSortingGPU.cpp b/Engine/Source/Runtime/Engine/Private/Particles/ParticleSortingGPU.cpp
--- a/Engine/Source/Runtime/Engine/Private/Particles/ParticleSortingGPU.cpp
+++ b/Engine/Source/Runtime/Engine/Private/Particles/ParticleSortingGPU.cpp
@@ -253,7 +253,7 @@ static int32 GenerateParticleSortKeys(
  */
 void FParticleSortBuffers::InitRHI()
 {
-	if (GetFeatureLevel() == ERHIFeatureLevel::SM5)
+	if (GetFeatureLevel() >= ERHIFeatureLevel::ES3_1)
 	{
 		for (int32 BufferIndex = 0; BufferIndex < 2; ++BufferIndex)
 		{
diff --git a/Engine/Source/Runtime/OpenGLDrv/OpenGLDrv.Build.cs b/Engine/Source/Runtime/OpenGLDrv/OpenGLDrv.Build.cs
--- a/Engine/Source/Runtime/OpenGLDrv/OpenGLDrv.Build.cs
+++ b/Engine/Source/Runtime/OpenGLDrv/OpenGLDrv.Build.cs
@@ -59,5 +59,7 @@ public class OpenGLDrv : ModuleRules
 		{
 			PrecompileForTargets = PrecompileTargetsType.None;
 		}
+
+		PublicDefinitions.Add(Target.Platform == UnrealTargetPlatform.Android ? "USE_ANDROID_OPENGL=1" : "USE_ANDROID_OPENGL=0");
 	}
 }
diff --git a/Engine/Source/Runtime/OpenGLDrv/Private/Android/AndroidEGL.cpp b/Engine/Source/Runtime/OpenGLDrv/Private/Android/AndroidEGL.cpp
--- a/Engine/Source/Runtime/OpenGLDrv/Private/Android/AndroidEGL.cpp
+++ b/Engine/Source/Runtime/OpenGLDrv/Private/Android/AndroidEGL.cpp
@@ -746,6 +746,8 @@ void AndroidEGL::InitSurface(bool bUseSmallSurface, bool bCreateWndSurface)
 	PImplData->SharedContext.eglSurface = PImplData->auxSurface;
 	PImplData->RenderingContext.eglSurface = PImplData->eglSurface;
 	PImplData->SingleThreadedContext.eglSurface = PImplData->eglSurface;
+	glGenVertexArrays = (PFNGLGENVERTEXARRAYSPROC)((void*)eglGetProcAddress("glGenVertexArrays"));
+	glBindVertexArray = (PFNGLBINDVERTEXARRAYPROC)((void*)eglGetProcAddress("glBindVertexArray"));
 }
 
 void AndroidEGL::ReInit()
@@ -1519,6 +1521,12 @@ bool AndroidEGL::InitContexts()
 	return Result;
 }
 
+void AndroidEGL::GenerateAndBindDefaultVAO(FPlatformOpenGLContext& InContext)
+{
+	glGenVertexArrays(1, &InContext.DefaultVertexArrayObject);
+	glBindVertexArray(InContext.DefaultVertexArrayObject);
+}
+
 void AndroidEGL::SetCurrentSharedContext()
 {
 	check(IsInGameThread());
@@ -1526,11 +1534,19 @@ void AndroidEGL::SetCurrentSharedContext()
 
 	if(GUseThreadedRendering)
 	{
-		SetCurrentContext(PImplData->SharedContext.eglContext, PImplData->SharedContext.eglSurface);
+		EGLBoolean Result = SetCurrentContext(PImplData->SharedContext.eglContext, PImplData->SharedContext.eglSurface);
+		if (Result == EGL_TRUE)
+		{
+			GenerateAndBindDefaultVAO(PImplData->SharedContext);
+		}
 	}
 	else
 	{
-		SetCurrentContext(PImplData->SingleThreadedContext.eglContext, PImplData->SingleThreadedContext.eglSurface);
+		EGLBoolean Result = SetCurrentContext(PImplData->SingleThreadedContext.eglContext, PImplData->SingleThreadedContext.eglSurface);
+		if (Result == EGL_TRUE)
+		{
+			GenerateAndBindDefaultVAO(PImplData->SingleThreadedContext);
+		}
 	}
 }
 
@@ -1544,11 +1560,19 @@ void AndroidEGL::SetCurrentRenderingContext()
 	PImplData->CurrentContextType = CONTEXT_Rendering;
 	if (GUseThreadedRendering)
 	{
-		SetCurrentContext(PImplData->RenderingContext.eglContext, PImplData->RenderingContext.eglSurface);
+		EGLBoolean Result = SetCurrentContext(PImplData->RenderingContext.eglContext, PImplData->RenderingContext.eglSurface);
+		if (Result == EGL_TRUE)
+		{
+			GenerateAndBindDefaultVAO(PImplData->RenderingContext);
+		};
 	}
 	else
 	{
-		SetCurrentContext(PImplData->SingleThreadedContext.eglContext, PImplData->SingleThreadedContext.eglSurface);
+		EGLBoolean Result = SetCurrentContext(PImplData->SingleThreadedContext.eglContext, PImplData->SingleThreadedContext.eglSurface);
+		if (Result == EGL_TRUE)
+		{
+			GenerateAndBindDefaultVAO(PImplData->SingleThreadedContext);
+		}
 	}
 }
 
diff --git a/Engine/Source/Runtime/OpenGLDrv/Private/Android/AndroidEGL.h b/Engine/Source/Runtime/OpenGLDrv/Private/Android/AndroidEGL.h
--- a/Engine/Source/Runtime/OpenGLDrv/Private/Android/AndroidEGL.h
+++ b/Engine/Source/Runtime/OpenGLDrv/Private/Android/AndroidEGL.h
@@ -88,6 +88,7 @@ public:
 	GLuint GetResolveFrameBuffer();
 	bool IsCurrentContextValid();
 	EGLContext  GetCurrentContext(  );
+	void GenerateAndBindDefaultVAO(FPlatformOpenGLContext& InContext);
 	void SetCurrentSharedContext();
 	void SetCurrentRenderingContext();
 	uint32_t GetCurrentContextType();
diff --git a/Engine/Source/Runtime/OpenGLDrv/Private/Android/AndroidOpenGL.cpp b/Engine/Source/Runtime/OpenGLDrv/Private/Android/AndroidOpenGL.cpp
--- a/Engine/Source/Runtime/OpenGLDrv/Private/Android/AndroidOpenGL.cpp
+++ b/Engine/Source/Runtime/OpenGLDrv/Private/Android/AndroidOpenGL.cpp
@@ -63,6 +63,14 @@ PFNGLGETOBJECTPTRLABELKHRPROC			glGetObjectPtrLabelKHR = NULL;
 
 PFNGLDRAWELEMENTSINSTANCEDPROC			glDrawElementsInstanced = NULL;
 PFNGLDRAWARRAYSINSTANCEDPROC			glDrawArraysInstanced = NULL;
+
+PFNGLGENVERTEXARRAYSPROC 				glGenVertexArrays = NULL;
+PFNGLBINDVERTEXARRAYPROC 				glBindVertexArray = NULL;
+PFNGLMAPBUFFERRANGEPROC					glMapBufferRange = NULL;
+PFNGLCOPYBUFFERSUBDATAPROC				glCopyBufferSubData = NULL;
+PFNGLDRAWARRAYSINDIRECTPROC				glDrawArraysIndirect = NULL;
+PFNGLDRAWELEMENTSINDIRECTPROC			glDrawElementsIndirect = NULL;
+
 PFNGLVERTEXATTRIBDIVISORPROC			glVertexAttribDivisor = NULL;
 
 PFNGLUNIFORM4UIVPROC					glUniform4uiv = NULL;
@@ -1085,6 +1093,8 @@ void FAndroidOpenGL::ProcessExtensions(const FString& ExtensionsString)
 
 	if (bES30Support)
 	{
+		glMapBufferRange = (PFNGLMAPBUFFERRANGEPROC)((void*)eglGetProcAddress("glMapBufferRange"));
+		glCopyBufferSubData = (PFNGLCOPYBUFFERSUBDATAPROC)((void*)eglGetProcAddress("glCopyBufferSubData"));
 		glDrawElementsInstanced = (PFNGLDRAWELEMENTSINSTANCEDPROC)((void*)eglGetProcAddress("glDrawElementsInstanced"));
 		glDrawArraysInstanced = (PFNGLDRAWARRAYSINSTANCEDPROC)((void*)eglGetProcAddress("glDrawArraysInstanced"));
 		glVertexAttribDivisor = (PFNGLVERTEXATTRIBDIVISORPROC)((void*)eglGetProcAddress("glVertexAttribDivisor"));
@@ -1153,6 +1163,8 @@ void FAndroidOpenGL::ProcessExtensions(const FString& ExtensionsString)
 
 	if (bES31Support)
 	{
+		glDrawArraysIndirect = (PFNGLDRAWARRAYSINDIRECTPROC)((void*)eglGetProcAddress("glDrawArraysIndirect"));
+		glDrawElementsIndirect = (PFNGLDRAWELEMENTSINDIRECTPROC)((void*)eglGetProcAddress("glDrawElementsIndirect"));
 		bSupportsTextureBuffer = ExtensionsString.Contains(TEXT("GL_EXT_texture_buffer"));
 		if (bSupportsTextureBuffer)
 		{
diff --git a/Engine/Source/Runtime/OpenGLDrv/Private/Android/AndroidOpenGL.h b/Engine/Source/Runtime/OpenGLDrv/Private/Android/AndroidOpenGL.h
--- a/Engine/Source/Runtime/OpenGLDrv/Private/Android/AndroidOpenGL.h
+++ b/Engine/Source/Runtime/OpenGLDrv/Private/Android/AndroidOpenGL.h
@@ -144,6 +144,13 @@ extern PFNGLDRAWELEMENTSINSTANCEDPROC	glDrawElementsInstanced;
 extern PFNGLDRAWARRAYSINSTANCEDPROC		glDrawArraysInstanced;
 extern PFNGLVERTEXATTRIBDIVISORPROC		glVertexAttribDivisor;
 
+extern PFNGLGENVERTEXARRAYSPROC 		glGenVertexArrays;
+extern PFNGLBINDVERTEXARRAYPROC 		glBindVertexArray;
+extern PFNGLMAPBUFFERRANGEPROC			glMapBufferRange;
+extern PFNGLCOPYBUFFERSUBDATAPROC		glCopyBufferSubData;
+extern PFNGLDRAWARRAYSINDIRECTPROC		glDrawArraysIndirect;
+extern PFNGLDRAWELEMENTSINDIRECTPROC	glDrawElementsIndirect;
+
 extern PFNGLTEXBUFFEREXTPROC			glTexBufferEXT;
 extern PFNGLUNIFORM4UIVPROC				glUniform4uiv;
 extern PFNGLCLEARBUFFERFIPROC			glClearBufferfi;
@@ -404,6 +411,21 @@ struct FAndroidOpenGL : public FOpenGLES2
 		glDrawElementsInstanced(Mode, Count, Type, Indices, InstanceCount);
 	}
 
+	static FORCEINLINE void CopyBufferSubData(GLenum ReadTarget, GLenum WriteTarget, GLintptr ReadOffset, GLintptr WriteOffset, GLsizeiptr Size)
+	{
+		glCopyBufferSubData(ReadTarget, WriteTarget, ReadOffset, WriteOffset, Size);
+	}
+
+	static FORCEINLINE void DrawArraysIndirect(GLenum Mode, const void *Offset)
+	{
+		glDrawArraysIndirect(Mode, Offset);
+	}
+
+	static FORCEINLINE void DrawElementsIndirect(GLenum Mode, GLenum Type, const void *Offset)
+	{
+		glDrawElementsIndirect(Mode, Type, Offset);
+	}
+
 	static FORCEINLINE void VertexAttribDivisor(GLuint Index, GLuint Divisor)
 	{
 		if (SupportsInstancing())
@@ -441,6 +463,41 @@ struct FAndroidOpenGL : public FOpenGLES2
 		}
 	}
 	
+	static FORCEINLINE void* MapBufferRange(GLenum Type, uint32 InOffset, uint32 InSize, EResourceLockMode LockMode)
+	{
+		GLenum Access;
+		switch (LockMode)
+		{
+		case EResourceLockMode::RLM_ReadOnly:
+			Access = GL_MAP_READ_BIT;
+			break;
+		case EResourceLockMode::RLM_WriteOnly:
+			Access = (GL_MAP_INVALIDATE_BUFFER_BIT | GL_MAP_WRITE_BIT | GL_MAP_UNSYNCHRONIZED_BIT);
+			break;
+		case EResourceLockMode::RLM_WriteOnlyUnsynchronized:
+			Access = (GL_MAP_WRITE_BIT | GL_MAP_UNSYNCHRONIZED_BIT);
+			break;
+		case EResourceLockMode::RLM_WriteOnlyPersistent:
+			Access = (GL_MAP_WRITE_BIT | GL_MAP_PERSISTENT_BIT | GL_MAP_COHERENT_BIT);
+			break;
+		case EResourceLockMode::RLM_ReadWrite:
+		default:
+			Access = (GL_MAP_READ_BIT | GL_MAP_WRITE_BIT);
+		}
+		return glMapBufferRange(Type, InOffset, InSize, Access);
+	}
+
+	static FORCEINLINE void UnmapBuffer(GLenum Type)
+	{
+		glUnmapBuffer(Type);
+	}
+
+	static FORCEINLINE void UnmapBufferRange(GLenum Type, uint32 InOffset, uint32 InSize)
+	{
+		UnmapBuffer(Type);
+	}
+
+
 	static FORCEINLINE void TexImage3D(GLenum Target, GLint Level, GLint InternalFormat, GLsizei Width, GLsizei Height, GLsizei Depth, GLint Border, GLenum Format, GLenum Type, const GLvoid* PixelData)
 	{
 		glTexImage3D(Target, Level, InternalFormat, Width, Height, Depth, Border, Format, Type, PixelData);
@@ -691,7 +748,9 @@ struct FAndroidOpenGL : public FOpenGLES2
 			glFramebufferFetchBarrierQCOM();
 		}
 	}
-			
+		
+	static FORCEINLINE bool SupportsDrawIndirect() { return true; }
+
 	static void ProcessExtensions(const FString& ExtensionsString);
 	static void SetupDefaultGLContextState(const FString& ExtensionsString);
 
diff --git a/Engine/Source/Runtime/OpenGLDrv/Private/OpenGLCommands.cpp b/Engine/Source/Runtime/OpenGLDrv/Private/OpenGLCommands.cpp
--- a/Engine/Source/Runtime/OpenGLDrv/Private/OpenGLCommands.cpp
+++ b/Engine/Source/Runtime/OpenGLDrv/Private/OpenGLCommands.cpp
@@ -3356,7 +3356,7 @@ void FOpenGLDynamicRHI::RHIInvalidateCachedState()
 
 void FOpenGLDynamicRHI::RHICopyToStagingBuffer(FRHIVertexBuffer* SourceBufferRHI, FRHIStagingBuffer* DestinationStagingBufferRHI, uint32 InOffset, uint32 InNumBytes)
 {
-#if OPENGL_GL3
+#if OPENGL_GL3 || USE_ANDROID_OPENGL
 	VERIFY_GL_SCOPE();
 	FOpenGLVertexBuffer* SourceBuffer = ResourceCast(SourceBufferRHI);
 	FOpenGLStagingBuffer* DestinationBuffer = ResourceCast(DestinationStagingBufferRHI);
diff --git a/Engine/Source/Runtime/OpenGLDrv/Private/OpenGLQuery.cpp b/Engine/Source/Runtime/OpenGLDrv/Private/OpenGLQuery.cpp
--- a/Engine/Source/Runtime/OpenGLDrv/Private/OpenGLQuery.cpp
+++ b/Engine/Source/Runtime/OpenGLDrv/Private/OpenGLQuery.cpp
@@ -1146,7 +1146,7 @@ void FOpenGLDisjointTimeStampQuery::ReleaseResources()
 
 FGPUFenceRHIRef FOpenGLDynamicRHI::RHICreateGPUFence(const FName &Name)
 {
-#if OPENGL_GL3
+#if OPENGL_GL3 || USE_ANDROID_OPENGL
 	return new FOpenGLGPUFence(Name);
 #else
 	UE_LOG(LogRHI, Fatal, TEXT("Fences are only available in OpenGL3 or later"));
@@ -1156,7 +1156,7 @@ FGPUFenceRHIRef FOpenGLDynamicRHI::RHICreateGPUFence(const FName &Name)
 
 FOpenGLGPUFence::~FOpenGLGPUFence()
 {
-#if OPENGL_GL3
+#if OPENGL_GL3 || USE_ANDROID_OPENGL
 	if (bValidSync)
 	{
 		FRHICommandListImmediate& RHICmdList = FRHICommandListExecutor::GetImmediateCommandList();
@@ -1172,7 +1172,7 @@ FOpenGLGPUFence::~FOpenGLGPUFence()
 
 void FOpenGLGPUFence::Clear()
 {
-#if OPENGL_GL3
+#if OPENGL_GL3 || USE_ANDROID_OPENGL
 
 	if (bValidSync)
 	{
@@ -1189,7 +1189,7 @@ void FOpenGLGPUFence::Clear()
 }
 bool FOpenGLGPUFence::Poll() const
 {
-#if OPENGL_GL3
+#if OPENGL_GL3 || USE_ANDROID_OPENGL
 	if (!bValidSync)
 	{
 		return false;
@@ -1211,7 +1211,7 @@ bool FOpenGLGPUFence::Poll() const
 
 void FOpenGLGPUFence::WriteInternal()
 {
-#if OPENGL_GL3
+#if OPENGL_GL3 || USE_ANDROID_OPENGL
 	FRHICommandListImmediate& RHICmdList = FRHICommandListExecutor::GetImmediateCommandList();
 	RHITHREAD_GLCOMMAND_PROLOGUE();
 	VERIFY_GL_SCOPE();
diff --git a/Engine/Source/Runtime/OpenGLDrv/Private/OpenGLVertexBuffer.cpp b/Engine/Source/Runtime/OpenGLDrv/Private/OpenGLVertexBuffer.cpp
--- a/Engine/Source/Runtime/OpenGLDrv/Private/OpenGLVertexBuffer.cpp
+++ b/Engine/Source/Runtime/OpenGLDrv/Private/OpenGLVertexBuffer.cpp
@@ -253,7 +253,7 @@ void FOpenGLDynamicRHI::RHICopyVertexBuffer(FRHIVertexBuffer* SourceBufferRHI, F
 
 FStagingBufferRHIRef FOpenGLDynamicRHI::RHICreateStagingBuffer()
 {
-#if OPENGL_GL3
+#if OPENGL_GL3 || USE_ANDROID_OPENGL
 	return new FOpenGLStagingBuffer();
 #else
 	UE_LOG(LogRHI, Fatal, TEXT("Staging Buffers are only available in OpenGL3 or later"));
@@ -285,14 +285,14 @@ FOpenGLStagingBuffer::~FOpenGLStagingBuffer()
 // I don't see a way to do this without stalling the RHI thread.
 void* FOpenGLStagingBuffer::Lock(uint32 Offset, uint32 NumBytes)
 {
-#if OPENGL_GL3
+#if OPENGL_GL3 || USE_ANDROID_OPENGL
 	FRHICommandListImmediate& RHICmdList = FRHICommandListExecutor::GetImmediateCommandList();
 	RHITHREAD_GLCOMMAND_PROLOGUE();
 	VERIFY_GL_SCOPE();
 
 	check(ShadowBuffer != 0);
 	glBindBuffer(GL_COPY_WRITE_BUFFER, ShadowBuffer);
-	void* Mapping = glMapBuffer(GL_COPY_WRITE_BUFFER, GL_READ_ONLY);
+	void* Mapping = FOpenGL::MapBufferRange(GL_COPY_WRITE_BUFFER, 0, NumBytes, FOpenGL::EResourceLockMode::RLM_ReadOnly);
 	check(Mapping);
 	return reinterpret_cast<uint8*>(Mapping) + Offset;
 
@@ -307,10 +307,10 @@ void* FOpenGLStagingBuffer::Lock(uint32 Offset, uint32 NumBytes)
 // Since this will probably be close to a call to lock we've probably paid most of the cost already.
 void FOpenGLStagingBuffer::Unlock()
 {
-#if OPENGL_GL3
+#if OPENGL_GL3 || USE_ANDROID_OPENGL
 	FRHICommandListImmediate& RHICmdList = FRHICommandListExecutor::GetImmediateCommandList();
 	RHITHREAD_GLCOMMAND_PROLOGUE();
-	glUnmapBuffer(GL_COPY_WRITE_BUFFER);
+	FOpenGL::UnmapBuffer(GL_COPY_WRITE_BUFFER);
 	glBindBuffer(GL_COPY_WRITE_BUFFER, 0);
 	RHITHREAD_GLCOMMAND_EPILOGUE();
 #else
diff --git a/Engine/Source/Runtime/OpenGLDrv/Public/OpenGL3.h b/Engine/Source/Runtime/OpenGLDrv/Public/OpenGL3.h
--- a/Engine/Source/Runtime/OpenGLDrv/Public/OpenGL3.h
+++ b/Engine/Source/Runtime/OpenGLDrv/Public/OpenGL3.h
@@ -37,6 +37,7 @@ struct FOpenGL3 : public FOpenGLBase
 	static FORCEINLINE bool IsDebugContent()						{ return bDebugContext; }
 
 	static FORCEINLINE bool SupportsTimestampQueries()				{ return TimestampQueryBits > 0; }
+	static FORCEINLINE bool SupportsDrawIndirect()					{ return true; }
 	static FORCEINLINE bool SupportsSeamlessCubeMap()				{ return bSupportsSeamlessCubemap; }
 	static FORCEINLINE bool SupportsVolumeTextureRendering()		{ return bSupportsVolumeTextureRendering; }
 	static FORCEINLINE bool SupportsGenerateMipmap()				{ return true; }
diff --git a/Engine/Source/Runtime/Renderer/Private/MobileShadingRenderer.cpp b/Engine/Source/Runtime/Renderer/Private/MobileShadingRenderer.cpp
--- a/Engine/Source/Runtime/Renderer/Private/MobileShadingRenderer.cpp
+++ b/Engine/Source/Runtime/Renderer/Private/MobileShadingRenderer.cpp
@@ -687,6 +687,20 @@ void FMobileSceneRenderer::Render(FRHICommandListImmediate& RHICmdList)
 	// End of scene color rendering
 	RHICmdList.EndRenderPass();
 
+	if (Scene->FXSystem && Views.IsValidIndex(0))
+	{
+		check(RHICmdList.IsOutsideRenderPass());
+
+		Scene->FXSystem->PostRenderOpaque(
+			RHICmdList,
+			Views[0].ViewUniformBuffer,
+			nullptr,
+			nullptr,
+			Views[0].AllowGPUParticleUpdate()
+		);
+		RHICmdList.ImmediateFlush(EImmediateFlushType::DispatchToRHIThread);
+	}
+
 	// Flush / submit cmdbuffer
 	if (bSubmitOffscreenRendering)
 	{
