Enable GL_OES_EGL_image_external_essl3 for samplerExternalOES where defined.

diff --git a/Engine/Source/Developer/ShaderFormatOpenGL/Private/GlslBackend.cpp b/Engine/Source/Developer/ShaderFormatOpenGL/Private/GlslBackend.cpp
--- a/Engine/Source/Developer/ShaderFormatOpenGL/Private/GlslBackend.cpp
+++ b/Engine/Source/Developer/ShaderFormatOpenGL/Private/GlslBackend.cpp
@@ -3106,10 +3106,17 @@ class ir_gen_glsl_visitor : public ir_visitor
 #endif		
 	}
 
-	void print_extensions(_mesa_glsl_parse_state* state, bool bUsesFramebufferFetchES2, bool bUsesDepthbufferFetchES2, bool bUsesES31Extensions, bool bInUsesExternalTexture)
+	void print_extensions(_mesa_glsl_parse_state* state, bool bUsesFramebufferFetchES2, bool bUsesDepthbufferFetchES2, bool bInUsesExternalTexture)
 	{
 		if (bInUsesExternalTexture)
 		{
+			if (CompileTarget == HCT_FeatureLevelES3_1)
+			{
+				ralloc_asprintf_append(buffer, "\n#ifdef GL_OES_EGL_image_external_essl3\n");
+				ralloc_asprintf_append(buffer, "#extension GL_OES_EGL_image_external_essl3 : enable\n");
+				ralloc_asprintf_append(buffer, "\n#endif\n");
+			}
+
 			ralloc_asprintf_append(buffer, "// Uses samplerExternalOES\n");
 		}
 
@@ -3155,7 +3162,7 @@ class ir_gen_glsl_visitor : public ir_visitor
 			ralloc_asprintf_append(buffer, "#extension GL_ARM_shader_framebuffer_fetch_depth_stencil : enable\n");
 		}
 
-		if (bUsesES31Extensions)
+		if (CompileTarget == HCT_FeatureLevelES3_1Ext)
 		{
 			ralloc_asprintf_append(buffer, "\n#ifdef GL_EXT_gpu_shader5\n");
 			ralloc_asprintf_append(buffer, "#extension GL_EXT_gpu_shader5 : enable\n");
@@ -3427,7 +3434,7 @@ bool compiler_internal_AdjustIsFrontFacing(bool isFrontFacing)
 
 		char* Extensions = ralloc_asprintf(mem_ctx, "");
 		buffer = &Extensions;
-		print_extensions(state, bUsesFrameBufferFetch, bUsesDepthbufferFetch, CompileTarget == HCT_FeatureLevelES3_1Ext, bUsesExternalTexture);
+		print_extensions(state, bUsesFrameBufferFetch, bUsesDepthbufferFetch, bUsesExternalTexture);
 		if (state->bSeparateShaderObjects && !(state->bGenerateES || CompileTarget == HCT_FeatureLevelES3_1))
 		{
 			switch (state->target)
