Incrementing Shader Version Num after previous backout to avoid DDC errors.
#tests none

diff --git a/Engine/Source/Developer/ShaderFormatOpenGL/Private/ShaderFormatOpenGL.cpp b/Engine/Source/Developer/ShaderFormatOpenGL/Private/ShaderFormatOpenGL.cpp
--- a/Engine/Source/Developer/ShaderFormatOpenGL/Private/ShaderFormatOpenGL.cpp
+++ b/Engine/Source/Developer/ShaderFormatOpenGL/Private/ShaderFormatOpenGL.cpp
@@ -23,8 +23,8 @@ class FShaderFormatGLSL : public IShaderFormat
 	enum
 	{
 		/** Version for shader format, this becomes part of the DDC key. */
-		UE_SHADER_GLSL_VER = 82,
-		UE_SHADER_GLSL_ANDROID_VER = 82,
+		UE_SHADER_GLSL_VER = 84,
+		UE_SHADER_GLSL_ANDROID_VER = 84,
 	};
 
 	void CheckFormat(FName Format) const
