Fixed: SingleLayerWater material does not show ES3.1 platform stats

diff --git a/Engine/Source/Developer/ShaderFormatOpenGL/Private/OpenGLShaderCompiler.cpp b/Engine/Source/Developer/ShaderFormatOpenGL/Private/OpenGLShaderCompiler.cpp
--- a/Engine/Source/Developer/ShaderFormatOpenGL/Private/OpenGLShaderCompiler.cpp
+++ b/Engine/Source/Developer/ShaderFormatOpenGL/Private/OpenGLShaderCompiler.cpp
@@ -2980,9 +2980,8 @@ TSharedPtr<ANSICHAR> FOpenGLFrontend::PrepareCodeForOfflineCompilation(const GLS
 	// This is the place to insert such engine preprocessor defines, immediately after the glsl version declaration.
 	if (Capabilities.bRequiresUEShaderFramebufferFetchDef && TypeEnum == GL_FRAGMENT_SHADER)
 	{
-		// Some devices (Zenfone5) support GL_EXT_shader_framebuffer_fetch but do not define GL_EXT_shader_framebuffer_fetch in GLSL compiler
-		// We can't define anything with GL_, so we use UE_EXT_shader_framebuffer_fetch to enable frame buffer fetch
-		StrOutSource.Append(TEXT("#define UE_EXT_shader_framebuffer_fetch 1\n"));
+		// Mali offline shader compiler does not support GL_EXT_shader_framebuffer_fetch
+		//StrOutSource.Append(TEXT("#define UE_EXT_shader_framebuffer_fetch 1\n"));
 	}
 
 	if (bEmitMobileMultiView)
