Fix for overflowing framebuffer fetch on mali GPUs when MSAA is enabled.
#rnx

diff --git a/Engine/Source/Developer/ShaderFormatOpenGL/Private/GlslBackend.cpp b/Engine/Source/Developer/ShaderFormatOpenGL/Private/GlslBackend.cpp
--- a/Engine/Source/Developer/ShaderFormatOpenGL/Private/GlslBackend.cpp
+++ b/Engine/Source/Developer/ShaderFormatOpenGL/Private/GlslBackend.cpp
@@ -3348,6 +3348,7 @@ bool compiler_internal_AdjustIsFrontFacing(bool isFrontFacing)
 				ralloc_asprintf_append(buffer, "	vec4 FramebufferFetchES2();\n");
 				ralloc_asprintf_append(buffer, "#elif defined( GL_ARM_shader_framebuffer_fetch)\n");
 				ralloc_asprintf_append(buffer, "	#define %sout\n", ES31FrameBufferFetchStorageQualifier);
+				ralloc_asprintf_append(buffer, "	highp vec4 gl_LastFragColorARM;\n");
 				ralloc_asprintf_append(buffer, "	vec4 FramebufferFetchES2() { return gl_LastFragColorARM; }\n");
 				ralloc_asprintf_append(buffer, "#else\n");
 				ralloc_asprintf_append(buffer, "	#define %sout\n", ES31FrameBufferFetchStorageQualifier);
@@ -3364,6 +3365,7 @@ bool compiler_internal_AdjustIsFrontFacing(bool isFrontFacing)
 				ralloc_asprintf_append(buffer, "	#endif\n");
 				ralloc_asprintf_append(buffer, "#else\n");
 				ralloc_asprintf_append(buffer, "	#ifdef GL_ARM_shader_framebuffer_fetch\n");
+				ralloc_asprintf_append(buffer, "		highp vec4 gl_LastFragColorARM;\n");
 				ralloc_asprintf_append(buffer, "		vec4 FramebufferFetchES2() { return gl_LastFragColorARM; }\n");
 				ralloc_asprintf_append(buffer, "	#else\n");
 				ralloc_asprintf_append(buffer, "		vec4 FramebufferFetchES2() { return vec4(65000.0, 65000.0, 65000.0, 65000.0); }\n");
diff --git a/Engine/Source/Developer/ShaderFormatOpenGL/Private/ShaderFormatOpenGL.cpp b/Engine/Source/Developer/ShaderFormatOpenGL/Private/ShaderFormatOpenGL.cpp
--- a/Engine/Source/Developer/ShaderFormatOpenGL/Private/ShaderFormatOpenGL.cpp
+++ b/Engine/Source/Developer/ShaderFormatOpenGL/Private/ShaderFormatOpenGL.cpp
@@ -23,8 +23,8 @@ class FShaderFormatGLSL : public IShaderFormat
 	enum
 	{
 		/** Version for shader format, this becomes part of the DDC key. */
-		UE_SHADER_GLSL_VER = 85,
-		UE_SHADER_GLSL_ANDROID_VER = 85,
+		UE_SHADER_GLSL_VER = 86,
+		UE_SHADER_GLSL_ANDROID_VER = 86,
 	};
 
 	void CheckFormat(FName Format) const
