Prevent OGL early fragment tests being enabled when discard instruction is used.
Fixes black outlines around masked materials.
#rnx

diff --git a/Engine/Source/Developer/ShaderFormatOpenGL/Private/GlslBackend.cpp b/Engine/Source/Developer/ShaderFormatOpenGL/Private/GlslBackend.cpp
--- a/Engine/Source/Developer/ShaderFormatOpenGL/Private/GlslBackend.cpp
+++ b/Engine/Source/Developer/ShaderFormatOpenGL/Private/GlslBackend.cpp
@@ -619,6 +619,9 @@ class ir_gen_glsl_visitor : public ir_visitor
 	// Found dFdx or dFdy
 	bool bUsesDXDY;
 
+	// True if the discard instruction was encountered.
+	bool bUsesDiscard;
+
 	// Uses gl_InstanceID
 	bool bUsesInstanceID;
 	
@@ -2072,6 +2075,7 @@ class ir_gen_glsl_visitor : public ir_visitor
 			ralloc_asprintf_append(buffer, ") ");
 		}
 		ralloc_asprintf_append(buffer, "discard");
+		bUsesDiscard = true;
 	}
 
 	bool try_conditional_move(ir_if *expr)
@@ -3032,7 +3036,7 @@ class ir_gen_glsl_visitor : public ir_visitor
 	 */
 	void print_layout(_mesa_glsl_parse_state *state)
 	{
-		if (early_depth_stencil)
+		if (early_depth_stencil && this->bUsesDiscard == false)
 		{
 			ralloc_asprintf_append(buffer, "layout(early_fragment_tests) in;\n");
 		}
@@ -3246,6 +3250,7 @@ public:
 		, bUsesTextureBuffer(false)
 		, bUseImageAtomic(false)
 		, bUsesDXDY(false)
+		, bUsesDiscard(false)
 		, bUsesInstanceID(false)
 		, bNoGlobalUniforms(bInNoGlobalUniforms)
 	{
diff --git a/Engine/Source/Developer/ShaderFormatOpenGL/Private/ShaderFormatOpenGL.cpp b/Engine/Source/Developer/ShaderFormatOpenGL/Private/ShaderFormatOpenGL.cpp
--- a/Engine/Source/Developer/ShaderFormatOpenGL/Private/ShaderFormatOpenGL.cpp
+++ b/Engine/Source/Developer/ShaderFormatOpenGL/Private/ShaderFormatOpenGL.cpp
@@ -20,7 +20,7 @@ class FShaderFormatGLSL : public IShaderFormat
 	enum
 	{
 		/** Version for shader format, this becomes part of the DDC key. */
-		UE_SHADER_GLSL_VER = 87,
+		UE_SHADER_GLSL_VER = 88,
 	};
 
 	void CheckFormat(FName Format) const
