Fixed: Assertion fail when launching TM-Shadermodels onto an Android Device with enabled eUB
And enabled eUB by default for OpenGL ES3.1

diff --git a/Engine/Source/Developer/ShaderFormatOpenGL/Private/OpenGLShaderCompiler.cpp b/Engine/Source/Developer/ShaderFormatOpenGL/Private/OpenGLShaderCompiler.cpp
--- a/Engine/Source/Developer/ShaderFormatOpenGL/Private/OpenGLShaderCompiler.cpp
+++ b/Engine/Source/Developer/ShaderFormatOpenGL/Private/OpenGLShaderCompiler.cpp
@@ -867,7 +867,13 @@ void FOpenGLFrontend::BuildShaderOutput(
 			Info.TypeIndex = CrossCompiler::PackedTypeNameToTypeIndex(TypeName);
 			InfoArray.Add(Info);
 		}
-
+		
+		// Sort by TypeIndex as expected by eUB uloading code
+		InfoArray.Sort([](const CrossCompiler::FPackedArrayInfo& A, const CrossCompiler::FPackedArrayInfo& B)
+		{ 
+			return A.TypeIndex < B.TypeIndex; 
+		});
+		
 		Header.Bindings.PackedUniformBuffers.Add(InfoArray);
 	}
 
diff --git a/Engine/Source/Developer/ShaderFormatOpenGL/Private/ShaderFormatOpenGL.cpp b/Engine/Source/Developer/ShaderFormatOpenGL/Private/ShaderFormatOpenGL.cpp
--- a/Engine/Source/Developer/ShaderFormatOpenGL/Private/ShaderFormatOpenGL.cpp
+++ b/Engine/Source/Developer/ShaderFormatOpenGL/Private/ShaderFormatOpenGL.cpp
@@ -18,7 +18,7 @@ class FShaderFormatGLSL : public IShaderFormat
 	enum
 	{
 		/** Version for shader format, this becomes part of the DDC key. */
-		UE_SHADER_GLSL_VER = 89,
+		UE_SHADER_GLSL_VER = 90,
 	};
 
 	void CheckFormat(FName Format) const
diff --git a/Engine/Source/Runtime/OpenGLDrv/Private/OpenGLCommands.cpp b/Engine/Source/Runtime/OpenGLDrv/Private/OpenGLCommands.cpp
--- a/Engine/Source/Runtime/OpenGLDrv/Private/OpenGLCommands.cpp
+++ b/Engine/Source/Runtime/OpenGLDrv/Private/OpenGLCommands.cpp
@@ -51,7 +51,7 @@ namespace OpenGLConsoleVariables
 
 	static FAutoConsoleVariable CVarUseEmulatedUBs(
 		TEXT("OpenGL.UseEmulatedUBs"),
-		0,
+		1,
 		TEXT("If true, enable using emulated uniform buffers on ES3.1 mode."),
 		ECVF_ReadOnly
 		);
