Force GLSL intrinsic return types to (full) float when HLSLCC_UseFullPrecisionInPS is requested.
#rnx

diff --git a/Engine/Source/Developer/ShaderFormatOpenGL/Private/GlslBackend.cpp b/Engine/Source/Developer/ShaderFormatOpenGL/Private/GlslBackend.cpp
--- a/Engine/Source/Developer/ShaderFormatOpenGL/Private/GlslBackend.cpp
+++ b/Engine/Source/Developer/ShaderFormatOpenGL/Private/GlslBackend.cpp
@@ -5677,8 +5677,9 @@ void FGlslCodeBackend::GenShaderPatchConstantFunctionInputs(_mesa_glsl_parse_sta
 
 void FGlslLanguageSpec::SetupLanguageIntrinsics(_mesa_glsl_parse_state* State, exec_list* ir)
 {
-	make_intrinsic_genType(ir, State, FRAMEBUFFER_FETCH_ES2, ir_invalid_opcode, IR_INTRINSIC_ALL_FLOATING, 0, 4, 4);
-	make_intrinsic_genType(ir, State, DEPTHBUFFER_FETCH_ES2, ir_invalid_opcode, IR_INTRINSIC_ALL_FLOATING, 3, 1, 1);
+	unsigned IntrinsicReturnType = bDefaultPrecisionIsHalf ? IR_INTRINSIC_ALL_FLOATING : IR_INTRINSIC_FLOAT;
+	make_intrinsic_genType(ir, State, FRAMEBUFFER_FETCH_ES2, ir_invalid_opcode, IntrinsicReturnType, 0, 4, 4);
+	make_intrinsic_genType(ir, State, DEPTHBUFFER_FETCH_ES2, ir_invalid_opcode, IntrinsicReturnType, 3, 1, 1);
 
 	{
 		ir_function* func = new(State)ir_function("compiler_internal_AdjustInputSemantic");
diff --git a/Engine/Source/Developer/ShaderFormatOpenGL/Private/OpenGLShaderCompiler.cpp b/Engine/Source/Developer/ShaderFormatOpenGL/Private/OpenGLShaderCompiler.cpp
--- a/Engine/Source/Developer/ShaderFormatOpenGL/Private/OpenGLShaderCompiler.cpp
+++ b/Engine/Source/Developer/ShaderFormatOpenGL/Private/OpenGLShaderCompiler.cpp
@@ -1198,6 +1198,11 @@ uint32 FOpenGLFrontend::CalculateCrossCompilerFlags(GLSLVersion Version, const T
 {
 	uint32  CCFlags = HLSLCC_NoPreprocess | HLSLCC_PackUniforms | HLSLCC_DX11ClipSpace | HLSLCC_RetainSizes;
 
+	if (CompilerFlags.Contains(CFLAG_UseFullPrecisionInPS))
+	{
+		CCFlags |= HLSLCC_UseFullPrecisionInPS;
+	}
+
 	if (CompilerFlags.Contains(CFLAG_FeatureLevelES31) ||
 		CompilerFlags.Contains(CFLAG_UseEmulatedUB))
 	{
@@ -1225,9 +1230,9 @@ FGlslCodeBackend* FOpenGLFrontend::CreateBackend(GLSLVersion Version, uint32 CCF
 	return new FGlslCodeBackend(CCFlags, HlslCompilerTarget);
 }
 
-FGlslLanguageSpec* FOpenGLFrontend::CreateLanguageSpec(GLSLVersion Version)
-{	
-	return new FGlslLanguageSpec();
+FGlslLanguageSpec* FOpenGLFrontend::CreateLanguageSpec(GLSLVersion Version, bool bDefaultPrecisionIsHalf)
+{
+	return new FGlslLanguageSpec(bDefaultPrecisionIsHalf);
 }
 
 #if USE_DXC
@@ -2475,7 +2480,9 @@ void FOpenGLFrontend::CompileShader(const FShaderCompilerInput& Input,FShaderCom
 		}
 
 		FGlslCodeBackend* BackEnd = CreateBackend(Version, CCFlags, HlslCompilerTarget);
-		FGlslLanguageSpec* LanguageSpec = CreateLanguageSpec(Version);
+
+		bool bDefaultPrecisionIsHalf = (CCFlags & HLSLCC_UseFullPrecisionInPS) == 0;
+		FGlslLanguageSpec* LanguageSpec = CreateLanguageSpec(Version, bDefaultPrecisionIsHalf);
 
 		FHlslCrossCompilerContext CrossCompilerContext(CCFlags, Frequency, HlslCompilerTarget);
 		if (CrossCompilerContext.Init(TCHAR_TO_ANSI(*Input.VirtualSourceFilePath), LanguageSpec))
diff --git a/Engine/Source/Developer/ShaderFormatOpenGL/Private/ShaderFormatOpenGL.cpp b/Engine/Source/Developer/ShaderFormatOpenGL/Private/ShaderFormatOpenGL.cpp
--- a/Engine/Source/Developer/ShaderFormatOpenGL/Private/ShaderFormatOpenGL.cpp
+++ b/Engine/Source/Developer/ShaderFormatOpenGL/Private/ShaderFormatOpenGL.cpp
@@ -18,7 +18,7 @@ class FShaderFormatGLSL : public IShaderFormat
 	enum
 	{
 		/** Version for shader format, this becomes part of the DDC key. */
-		UE_SHADER_GLSL_VER = 91,
+		UE_SHADER_GLSL_VER = 92,
 	};
 
 	void CheckFormat(FName Format) const
diff --git a/Engine/Source/Developer/ShaderFormatOpenGL/Public/GlslBackend.h b/Engine/Source/Developer/ShaderFormatOpenGL/Public/GlslBackend.h
--- a/Engine/Source/Developer/ShaderFormatOpenGL/Public/GlslBackend.h
+++ b/Engine/Source/Developer/ShaderFormatOpenGL/Public/GlslBackend.h
@@ -8,8 +8,12 @@
 
 class SHADERFORMATOPENGL_API FGlslLanguageSpec : public ILanguageSpec
 {
+protected:
+	bool bDefaultPrecisionIsHalf;
+
 public:
-	FGlslLanguageSpec() 
+	FGlslLanguageSpec(bool bInDefaultPrecisionIsHalf)
+		: bDefaultPrecisionIsHalf(bInDefaultPrecisionIsHalf)
 	{}
 
 	virtual bool SupportsDeterminantIntrinsic() const override
diff --git a/Engine/Source/Developer/ShaderFormatOpenGL/Public/ShaderFormatOpenGL.h b/Engine/Source/Developer/ShaderFormatOpenGL/Public/ShaderFormatOpenGL.h
--- a/Engine/Source/Developer/ShaderFormatOpenGL/Public/ShaderFormatOpenGL.h
+++ b/Engine/Source/Developer/ShaderFormatOpenGL/Public/ShaderFormatOpenGL.h
@@ -70,7 +70,7 @@ protected:
 	virtual struct FGlslCodeBackend* CreateBackend(GLSLVersion Version, uint32 CCFlags, EHlslCompileTarget HlslCompilerTarget);
 
 	// create the language spec
-	virtual class FGlslLanguageSpec* CreateLanguageSpec(GLSLVersion Version);
+	virtual class FGlslLanguageSpec* CreateLanguageSpec(GLSLVersion Version, bool bDefaultPrecisionIsHalf);
 
 
 	// Allow a subclass to perform additional work on the cross compiled source code
