Disabled gl_LastFragColorARM re-declaration as it does not work on S7 Mali

diff --git a/Engine/Source/Developer/ShaderFormatOpenGL/Private/GlslBackend.cpp b/Engine/Source/Developer/ShaderFormatOpenGL/Private/GlslBackend.cpp
--- a/Engine/Source/Developer/ShaderFormatOpenGL/Private/GlslBackend.cpp
+++ b/Engine/Source/Developer/ShaderFormatOpenGL/Private/GlslBackend.cpp
@@ -3344,7 +3344,7 @@ bool compiler_internal_AdjustIsFrontFacing(bool isFrontFacing)
 				ralloc_asprintf_append(buffer, "	vec4 FramebufferFetchES2();\n");
 				ralloc_asprintf_append(buffer, "#elif defined( GL_ARM_shader_framebuffer_fetch)\n");
 				ralloc_asprintf_append(buffer, "	#define %sout\n", ES31FrameBufferFetchStorageQualifier);
-				ralloc_asprintf_append(buffer, "	highp vec4 gl_LastFragColorARM;\n");
+				//ralloc_asprintf_append(buffer, "	highp vec4 gl_LastFragColorARM;\n"); // Does not work on S7 MALI
 				ralloc_asprintf_append(buffer, "	vec4 FramebufferFetchES2() { return gl_LastFragColorARM; }\n");
 				ralloc_asprintf_append(buffer, "#else\n");
 				ralloc_asprintf_append(buffer, "	#define %sout\n", ES31FrameBufferFetchStorageQualifier);
@@ -3361,7 +3361,7 @@ bool compiler_internal_AdjustIsFrontFacing(bool isFrontFacing)
 				ralloc_asprintf_append(buffer, "	#endif\n");
 				ralloc_asprintf_append(buffer, "#else\n");
 				ralloc_asprintf_append(buffer, "	#ifdef GL_ARM_shader_framebuffer_fetch\n");
-				ralloc_asprintf_append(buffer, "		highp vec4 gl_LastFragColorARM;\n");
+				//ralloc_asprintf_append(buffer, "		highp vec4 gl_LastFragColorARM;\n"); Does not work on S7 MALI
 				ralloc_asprintf_append(buffer, "		vec4 FramebufferFetchES2() { return gl_LastFragColorARM; }\n");
 				ralloc_asprintf_append(buffer, "	#else\n");
 				ralloc_asprintf_append(buffer, "		vec4 FramebufferFetchES2() { return vec4(65000.0, 65000.0, 65000.0, 65000.0); }\n");
