F - Copy of 12659574 for GLSL
Fix Vulkan codegen for RWStructuredBuffer<float4> not generating scalar access
#rnx

diff --git a/Engine/Source/Developer/ShaderFormatOpenGL/Private/GlslBackend.cpp b/Engine/Source/Developer/ShaderFormatOpenGL/Private/GlslBackend.cpp
--- a/Engine/Source/Developer/ShaderFormatOpenGL/Private/GlslBackend.cpp
+++ b/Engine/Source/Developer/ShaderFormatOpenGL/Private/GlslBackend.cpp
@@ -1694,7 +1694,7 @@ class ir_gen_glsl_visitor : public ir_visitor
 		}
 	}
 
-	void print_image_op( ir_dereference_image *deref, ir_rvalue *src)
+	void print_image_op( ir_dereference_image *deref, ir_rvalue *src, const char* Mask = nullptr)
 	{
 		const char* swizzle[] =
 		{
@@ -1725,6 +1725,10 @@ class ir_gen_glsl_visitor : public ir_visitor
 				ralloc_asprintf_append(buffer, "[");
 				deref->image_index->accept(this);
 				ralloc_asprintf_append(buffer, "]");
+				if (Mask)
+				{
+					ralloc_asprintf_append(buffer, Mask);
+				}
 				if (src)
 				{
 					ralloc_asprintf_append(buffer, " = ");
@@ -1840,31 +1844,30 @@ class ir_gen_glsl_visitor : public ir_visitor
 			ralloc_asprintf_append(buffer, ") { ");
 		}
 
-		if (assign->lhs->as_dereference_image() != NULL)
-		{
-			/** EHart - should the write mask be checked here? */
-			print_image_op( assign->lhs->as_dereference_image(), assign->rhs);
-		}
-		else
+		char mask[6] = "";
+		unsigned j = 1;
+		if (assign->lhs->type->is_scalar() == false ||
+			assign->write_mask != 0x1)
 		{
-			char mask[6];
-			unsigned j = 1;
-			if (assign->lhs->type->is_scalar() == false ||
-				assign->write_mask != 0x1)
+			for (unsigned i = 0; i < 4; i++)
 			{
-				for (unsigned i = 0; i < 4; i++)
+				if ((assign->write_mask & (1 << i)) != 0)
 				{
-					if ((assign->write_mask & (1 << i)) != 0)
-					{
-						mask[j] = "xyzw"[i];
-						j++;
-					}
+					mask[j] = "xyzw"[i];
+					j++;
 				}
 			}
-			mask[j] = '\0';
+		}
+		mask[j] = '\0';
 
-			mask[0] = (j == 1) ? '\0' : '.';
+		mask[0] = (j == 1) ? '\0' : '.';
 
+		if (assign->lhs->as_dereference_image() != NULL)
+		{
+			print_image_op( assign->lhs->as_dereference_image(), assign->rhs, mask);
+		}
+		else
+		{
 			// decide if we need to cast to float
 			const bool need_float_conv = (assign->lhs->type->is_float()
 				&& ((assign->rhs->as_constant() != nullptr)
