Initialize FRHIShader::ShaderName for OpenGL shaders when r.Shaders.KeepDebugInfo is on

diff --git a/Engine/Source/Developer/ShaderFormatOpenGL/Private/OpenGLShaderCompiler.cpp b/Engine/Source/Developer/ShaderFormatOpenGL/Private/OpenGLShaderCompiler.cpp
--- a/Engine/Source/Developer/ShaderFormatOpenGL/Private/OpenGLShaderCompiler.cpp
+++ b/Engine/Source/Developer/ShaderFormatOpenGL/Private/OpenGLShaderCompiler.cpp
@@ -975,11 +975,6 @@ void FOpenGLFrontend::BuildShaderOutput(
 			ShaderOutput.bSucceeded = true;
 		}
 
-		// store data we can pickup later with ShaderCode.FindOptionalData('n'), could be removed for shipping
-		// Daniel L: This GenerateShaderName does not generate a deterministic output among shaders as the shader code can be shared.
-		//			uncommenting this will cause the project to have non deterministic materials and will hurt patch sizes
-		//ShaderOutput.ShaderCode.AddOptionalData('n', TCHAR_TO_UTF8(*ShaderInput.GenerateShaderName()));
-
 		// extract final source code as requested by the Material Editor
 		if (ShaderInput.ExtraSettings.bExtractShaderSource)
 		{
@@ -988,6 +983,11 @@ void FOpenGLFrontend::BuildShaderOutput(
 			ShaderOutput.OptionalFinalShaderSource = FString(GlslCodeOriginal.GetData());
 		}
 
+		if (ShaderInput.Environment.CompilerFlags.Contains(CFLAG_KeepDebugInfo))
+		{
+			ShaderOutput.ShaderCode.AddOptionalData(FShaderCodeName::Key, TCHAR_TO_UTF8(*ShaderInput.GenerateShaderName()));
+		}
+
 		// if available, attempt run an offline compilation and extract statistics
 		if (ShaderInput.ExtraSettings.OfflineCompilerPath.Len() > 0)
 		{
diff --git a/Engine/Source/Runtime/OpenGLDrv/Private/OpenGLShaders.cpp b/Engine/Source/Runtime/OpenGLDrv/Private/OpenGLShaders.cpp
--- a/Engine/Source/Runtime/OpenGLDrv/Private/OpenGLShaders.cpp
+++ b/Engine/Source/Runtime/OpenGLDrv/Private/OpenGLShaders.cpp
@@ -782,6 +782,13 @@ ShaderType* CompileOpenGLShader(TArrayView<const uint8> InShaderCode, const FSHA
 		}
 	}
 
+#if !(UE_BUILD_SHIPPING || UE_BUILD_TEST)
+	if (RHIShader)
+	{
+		RHIShader->ShaderName = ShaderCode.FindOptionalData(FShaderCodeName::Key);
+	}
+#endif
+
 #if DEBUG_GL_SHADERS
 	Shader->GlslCode = GlslCode;
 	Shader->GlslCodeString = (ANSICHAR*)Shader->GlslCode.GetData();
