Dump GL compiled shaders for non-DXC path

diff --git a/Engine/Source/Developer/ShaderFormatOpenGL/Private/OpenGLShaderCompiler.cpp b/Engine/Source/Developer/ShaderFormatOpenGL/Private/OpenGLShaderCompiler.cpp
--- a/Engine/Source/Developer/ShaderFormatOpenGL/Private/OpenGLShaderCompiler.cpp
+++ b/Engine/Source/Developer/ShaderFormatOpenGL/Private/OpenGLShaderCompiler.cpp
@@ -2436,6 +2436,11 @@ void FOpenGLFrontend::CompileShader(const FShaderCompilerInput& Input, FShaderCo
 
 		delete BackEnd;
 		delete LanguageSpec;
+
+		if (bDumpDebugInfo && bCompilationSucceeded && GlslShaderSource != nullptr)
+		{
+			DumpDebugShaderText(Input, GlslShaderSource, FCStringAnsi::Strlen(GlslShaderSource), GetFrequencyFileExt(Frequency));
+		}
 	}
 
 	static const bool bDirectCompile = FParse::Param(FCommandLine::Get(), TEXT("directcompile"));
