Added support for Framebuffer Fetch and Pixel Local Storage to OpenGL platform.

diff --git a/Engine/Shaders/Private/MobileBasePassPixelShader.usf b/Engine/Shaders/Private/MobileBasePassPixelShader.usf
--- a/Engine/Shaders/Private/MobileBasePassPixelShader.usf
+++ b/Engine/Shaders/Private/MobileBasePassPixelShader.usf
@@ -525,11 +525,15 @@ void Main(
 	, FMobileBasePassInterpolantsVSToPS BasePassInterpolants
 	, in float4 SvPosition : SV_Position
 	OPTIONAL_IsFrontFace
-	, out half4 OutColor	: SV_Target0
 #if DEFERRED_SHADING_PATH
+#if !USE_GLES_FBF_DEFERRED
+	, out half4 OutColor	: SV_Target0
 	, out half4 OutGBufferA	: SV_Target1
 	, out half4 OutGBufferB	: SV_Target2
-	, out half4 OutGBufferC	: SV_Target3
+	, out half4 OutGBufferC	: SV_Target3  
+#endif
+#else
+	, out half4 OutColor	: SV_Target0
 #endif
 #if USE_SCENE_DEPTH_AUX
 	, out float OutSceneDepthAux : SV_Target4
@@ -539,6 +543,15 @@ void Main(
 #endif
 	)
 {
+
+#if DEFERRED_SHADING_PATH && USE_GLES_FBF_DEFERRED
+	half4 OutColor;	
+	half4 OutGBufferA;
+	half4 OutGBufferB;
+	half4 OutGBufferC;
+	float OutSceneDepthAux;
+#endif
+  
 #if MOBILE_MULTI_VIEW
 	ResolvedView = ResolveView(uint(BasePassInterpolants.MultiViewId));
 #else
@@ -676,6 +689,13 @@ void Main(
 	GBuffer.PrecomputedShadowFactors.r = Shadow;
 
 	EncodeGBuffer(GBuffer, OutGBufferA, OutGBufferB, OutGBufferC, OutGBufferD, OutGBufferE, OutGBufferF, OutGBufferVelocity);
+
+#if USE_GLES_FBF_DEFERRED
+	GLFramebufferFetchWrite1(OutGBufferA);  
+	GLFramebufferFetchWrite2(OutGBufferB);  
+	GLFramebufferFetchWrite3(OutGBufferC);
+	GLFramebufferFetchDepthWrite(SvPosition.z);   
+#endif
 #else
 
 #if !MATERIAL_SHADINGMODEL_UNLIT
@@ -966,4 +986,8 @@ void Main(
 		View.VTFeedbackBuffer
 	);
 #endif
+
+#if DEFERRED_SHADING_PATH && USE_GLES_FBF_DEFERRED 
+	GLFramebufferFetchWrite0(half4(OutColor.rgb, 1.0f));
+#endif
 }
diff --git a/Engine/Shaders/Private/MobileDeferredShading.usf b/Engine/Shaders/Private/MobileDeferredShading.usf
--- a/Engine/Shaders/Private/MobileDeferredShading.usf
+++ b/Engine/Shaders/Private/MobileDeferredShading.usf
@@ -64,6 +64,12 @@ void FetchGBuffer(in float2 UV, out float4 GBufferA, out float4 GBufferB, out fl
 	GBufferC = SubpassFetchRGBA_3(); 
 	GBufferD = 0; 
 	SceneDepth = ConvertFromDeviceZ(SubpassFetchR_4());
+#elif USE_GLES_FBF_DEFERRED
+	GBufferA = GLFramebufferFetch1(); 
+	GBufferB = GLFramebufferFetch2(); 
+	GBufferC = GLFramebufferFetch3();  
+	GBufferD = 0;   
+	SceneDepth = ConvertFromDeviceZ(GLFramebufferFetchDepth());
 #else
 	GBufferA = Texture2DSampleLevel(MobileSceneTextures.GBufferATexture, MobileSceneTextures.GBufferATextureSampler, UV, 0); 
 	GBufferB = Texture2DSampleLevel(MobileSceneTextures.GBufferBTexture, MobileSceneTextures.GBufferBTextureSampler, UV, 0);
diff --git a/Engine/Shaders/Private/MobileDeferredUtils.usf b/Engine/Shaders/Private/MobileDeferredUtils.usf
new file mode 100644
--- /dev/null
+++ b/Engine/Shaders/Private/MobileDeferredUtils.usf
@@ -0,0 +1,26 @@
+// Copyright Epic Games, Inc. All Rights Reserved.
+#include "Common.ush"
+
+void MobileDeferredCopyPLSPS(
+	noperspective float4 UVAndScreenPos : TEXCOORD0, 
+	float4 SvPosition : SV_POSITION, 
+	out half4 OutColor : SV_Target0)
+{
+#if USE_GLES_FBF_DEFERRED
+	OutColor = float4(GLFramebufferFetch0(), 1.0f);
+#else
+	OutColor = 1.0f;  
+#endif
+}
+
+void MobileDeferredCopyDepthPS(
+	noperspective float4 UVAndScreenPos : TEXCOORD0, 
+	float4 SvPosition : SV_POSITION, 
+	out float OutSceneDepthAux : SV_Target0)
+{
+#if USE_GLES_FBF_DEFERRED
+	OutSceneDepthAux = GLFramebufferFetchDepth();
+#else
+	OutSceneDepthAux = 1.0f;  
+#endif
+}
diff --git a/Engine/Source/Developer/Android/AndroidTargetPlatform/Private/AndroidTargetPlatform.cpp b/Engine/Source/Developer/Android/AndroidTargetPlatform/Private/AndroidTargetPlatform.cpp
--- a/Engine/Source/Developer/Android/AndroidTargetPlatform/Private/AndroidTargetPlatform.cpp
+++ b/Engine/Source/Developer/Android/AndroidTargetPlatform/Private/AndroidTargetPlatform.cpp
@@ -593,7 +593,7 @@ void FAndroidTargetPlatform::GetReflectionCaptureFormats( TArray<FName>& OutForm
 	static auto* MobileShadingPathCvar = IConsoleManager::Get().FindTConsoleVariableDataInt(TEXT("r.Mobile.ShadingPath"));
 	const bool bMobileDeferredShading = (MobileShadingPathCvar->GetValueOnAnyThread() == 1);
 	
-	if (SupportsVulkanSM5() || (SupportsVulkan() && bMobileDeferredShading))
+	if (SupportsVulkanSM5() || bMobileDeferredShading)
 	{
 		// use Full HDR with SM5 and Mobile Deferred
 		OutFormats.Add(FName(TEXT("FullHDR")));
diff --git a/Engine/Source/Developer/ShaderFormatOpenGL/Private/GlslBackend.cpp b/Engine/Source/Developer/ShaderFormatOpenGL/Private/GlslBackend.cpp
--- a/Engine/Source/Developer/ShaderFormatOpenGL/Private/GlslBackend.cpp
+++ b/Engine/Source/Developer/ShaderFormatOpenGL/Private/GlslBackend.cpp
@@ -333,6 +333,44 @@ static const char* GLSLIntCastTypes[5] =
 	"ivec4",
 };
 
+const char* GL_FRAMEBUFFER_FETCH[5] =
+{
+	"GLFramebufferFetch0",
+	"GLFramebufferFetch1",
+	"GLFramebufferFetch2",
+	"GLFramebufferFetch3",
+	"GLFramebufferFetchDepth"
+};
+
+const char* GL_FRAMEBUFFER_FETCH_WRITE[5] =
+{
+	"GLFramebufferFetchWrite0",
+	"GLFramebufferFetchWrite1",
+	"GLFramebufferFetchWrite2",
+	"GLFramebufferFetchWrite3",
+	"GLFramebufferFetchDepthWrite"
+};
+
+const char* GL_PLS_IMAGE_FORMAT[5] =
+{
+	"r11f_g11f_b10f",
+	"rgba8",
+	"rgba8",
+	"rgba8",
+	""
+};
+
+const char* GL_FRAMEBUFFER_FETCH_TYPE[5] =
+{
+	"vec3",
+	"vec4",
+	"vec4",
+	"vec4",
+	"float"
+};
+
+static const uint32_t FRAMEBUFFER_FETCH_DEPTH_INDEX = 4;
+
 static const char* ES31FrameBufferFetchStorageQualifier = "FRAME_BUFFERFETCH_STORAGE_QUALIFIER ";
 static_assert((sizeof(GLSLExpressionTable) / sizeof(GLSLExpressionTable[0])) == ir_opcode_count, "GLSLExpressionTableSizeMismatch");
 
@@ -560,6 +598,10 @@ class ir_gen_glsl_visitor : public ir_visitor
 	// depthbuffer fetch is in use
 	bool bUsesDepthbufferFetch;
 
+	// Mask for indicies using FBF
+	uint32 FramebufferFetchMask;
+	uint32 FramebufferFetchWriteMask;
+
 	// uses external texture
 	bool bUsesExternalTexture;
 
@@ -3151,8 +3193,23 @@ class ir_gen_glsl_visitor : public ir_visitor
 			ralloc_asprintf_append(buffer, "#extension GL_ARM_shader_framebuffer_fetch : enable\n");
 			ralloc_asprintf_append(buffer, "#endif\n");
 		}
+		else if(FramebufferFetchMask || FramebufferFetchWriteMask)
+		{
+			ralloc_asprintf_append(buffer, "#ifdef UE_MRT_FRAMEBUFFER_FETCH\n");
+			ralloc_asprintf_append(buffer, "#extension GL_EXT_shader_framebuffer_fetch : enable\n");
+			ralloc_asprintf_append(buffer, "#endif\n");
+
+			ralloc_asprintf_append(buffer, "#ifdef UE_MRT_PLS\n");
+			ralloc_asprintf_append(buffer, "#extension GL_EXT_shader_pixel_local_storage : enable\n");
+			ralloc_asprintf_append(buffer, "#endif\n");	
+		}
+
+		// Cube map arrays are required for deferred rendering on mobile
+		ralloc_asprintf_append(buffer, "\n#if (defined(UE_MRT_PLS) || defined(UE_MRT_FRAMEBUFFER_FETCH)) && defined(GL_EXT_texture_cube_map_array)\n");
+		ralloc_asprintf_append(buffer, "#extension GL_EXT_texture_cube_map_array : enable\n");
+		ralloc_asprintf_append(buffer, "\n#endif\n");
 
-		if (bUsesDepthbufferFetchES2)
+		if (bUsesDepthbufferFetchES2 || FramebufferFetchMask || FramebufferFetchWriteMask)
 		{
 			ralloc_asprintf_append(buffer, "#extension GL_ARM_shader_framebuffer_fetch_depth_stencil : enable\n");
 		}
@@ -3214,7 +3271,9 @@ class ir_gen_glsl_visitor : public ir_visitor
 public:
 
 	/** Constructor. */
-	ir_gen_glsl_visitor(bool bInIsES, bool bInEmitPrecision, bool bInIsWebGL, EHlslCompileTarget InCompileTarget, _mesa_glsl_parser_targets InShaderTarget, bool bInGenerateLayoutLocations, bool bInDefaultPrecisionIsHalf, bool bInNoGlobalUniforms, bool bInUsesFrameBufferFetch, bool bInUsesDepthbufferFetch, bool bInUsesExternalTexture)
+	ir_gen_glsl_visitor(bool bInIsES, bool bInEmitPrecision, bool bInIsWebGL, EHlslCompileTarget InCompileTarget, _mesa_glsl_parser_targets InShaderTarget,
+						bool bInGenerateLayoutLocations, bool bInDefaultPrecisionIsHalf, bool bInNoGlobalUniforms, bool bInUsesFrameBufferFetch, bool bInUsesDepthbufferFetch,
+						uint32_t InFramebufferFetchMask, uint32_t InFramebufferFetchWriteMask, bool bInUsesExternalTexture)
 		: early_depth_stencil(false)
 		, bIsES(bInIsES)
 		, bEmitPrecision(bInEmitPrecision)
@@ -3225,6 +3284,8 @@ public:
 		, bDefaultPrecisionIsHalf(bInDefaultPrecisionIsHalf)
 		, bUsesFrameBufferFetch(bInUsesFrameBufferFetch)
 		, bUsesDepthbufferFetch(bInUsesDepthbufferFetch)
+		, FramebufferFetchMask(InFramebufferFetchMask)
+		, FramebufferFetchWriteMask(InFramebufferFetchWriteMask)
 		, bUsesExternalTexture(bInUsesExternalTexture)
 		, buffer(0)
 		, indentation(0)
@@ -3371,8 +3432,6 @@ bool compiler_internal_AdjustIsFrontFacing(bool isFrontFacing)
 				ralloc_asprintf_append(buffer, "	#endif\n");
 				ralloc_asprintf_append(buffer, "#endif\n\n");
 			}
-
-
 		}
 
 		if (bUsesDepthbufferFetch)
@@ -3384,6 +3443,98 @@ bool compiler_internal_AdjustIsFrontFacing(bool isFrontFacing)
 			ralloc_asprintf_append(buffer, "#endif\n\n");
 		}
 
+		if (FramebufferFetchMask || FramebufferFetchWriteMask)
+		{
+			// Framebuffer Depth Fetch
+			ralloc_asprintf_append(buffer, "highp %s %s()\n"
+				"{\n"
+				"\treturn gl_LastFragDepthARM;\n"
+				"}\n", GL_FRAMEBUFFER_FETCH_TYPE[FRAMEBUFFER_FETCH_DEPTH_INDEX], GL_FRAMEBUFFER_FETCH[FRAMEBUFFER_FETCH_DEPTH_INDEX]);
+
+			ralloc_asprintf_append(buffer, "void %s(%s input_vec)\n"
+				"{\n"
+				"\tgl_FragDepth = input_vec;"
+				"}\n", GL_FRAMEBUFFER_FETCH_WRITE[FRAMEBUFFER_FETCH_DEPTH_INDEX], GL_FRAMEBUFFER_FETCH_TYPE[FRAMEBUFFER_FETCH_DEPTH_INDEX]);
+
+			// Framebuffer Fetch
+			ralloc_asprintf_append(buffer, "#ifdef UE_MRT_FRAMEBUFFER_FETCH\n");
+
+			for (int32_t i = 0; i < UE_ARRAY_COUNT(GL_FRAMEBUFFER_FETCH); ++i)
+			{
+				if (FramebufferFetchMask & (1 << i) || FramebufferFetchWriteMask & (1 << i))
+				{
+					if (i == FRAMEBUFFER_FETCH_DEPTH_INDEX)
+					{
+						continue;
+					}
+					
+					ralloc_asprintf_append(buffer, "layout(location = %d) inout %s out_Target%d;\n", i, GL_FRAMEBUFFER_FETCH_TYPE[i], i);
+					ralloc_asprintf_append(buffer, "highp %s %s()\n"
+						"{\n"
+						"\treturn out_Target%d;\n"
+						"}\n", GL_FRAMEBUFFER_FETCH_TYPE[i], GL_FRAMEBUFFER_FETCH[i], i);
+
+					ralloc_asprintf_append(buffer, "void %s(%s input_vec)\n"
+						"{\n"
+						"\tout_Target%d = input_vec;\n"
+						"}\n", GL_FRAMEBUFFER_FETCH_WRITE[i], GL_FRAMEBUFFER_FETCH_TYPE[i], i);
+				}
+			}
+
+			ralloc_asprintf_append(buffer, "#endif\n");
+			// End Framebuffer Fetch
+
+			// Pixel Local Storage
+			ralloc_asprintf_append(buffer, "#ifdef UE_MRT_PLS\n");
+		
+			if (FramebufferFetchMask && FramebufferFetchWriteMask)
+			{
+				ralloc_asprintf_append(buffer, "__pixel_localEXT PLSGBufferData\n{ \n");
+			}
+			else if (FramebufferFetchMask)
+			{
+				ralloc_asprintf_append(buffer, "__pixel_local_inEXT PLSGBufferData\n{ \n");
+			}
+			else
+			{
+				ralloc_asprintf_append(buffer, "__pixel_local_outEXT PLSGBufferData\n{ \n");
+			}
+
+			for (int32_t i = 0; i < FRAMEBUFFER_FETCH_DEPTH_INDEX; ++i)
+			{
+				ralloc_asprintf_append(buffer, "layout(%s) highp %s out_Target%d;\n", GL_PLS_IMAGE_FORMAT[i], GL_FRAMEBUFFER_FETCH_TYPE[i], i);
+			}
+
+			ralloc_asprintf_append(buffer, "} PLSGBuffer;\n\n");
+
+			for (int32_t i = 0; i < UE_ARRAY_COUNT(GL_FRAMEBUFFER_FETCH); ++i)
+			{
+				if (i == FRAMEBUFFER_FETCH_DEPTH_INDEX)
+				{
+					continue;
+				}
+
+				if (FramebufferFetchMask)
+				{
+					ralloc_asprintf_append(buffer, "highp %s %s()\n"
+						"{\n"
+						"\treturn PLSGBuffer.out_Target%d;\n"
+						"}\n", GL_FRAMEBUFFER_FETCH_TYPE[i], GL_FRAMEBUFFER_FETCH[i], i);
+				}
+
+				if(FramebufferFetchWriteMask)
+				{
+						ralloc_asprintf_append(buffer, "void %s(%s input_vec)\n"
+							"{\n"
+							"\tPLSGBuffer.out_Target%d = input_vec;\n"
+							"}\n", GL_FRAMEBUFFER_FETCH_WRITE[i], GL_FRAMEBUFFER_FETCH_TYPE[i], i);
+				}
+			}
+
+			ralloc_asprintf_append(buffer, "#endif\n");
+		}
+		// End Pixel Local Storage
+
 		foreach_iter(exec_list_iterator, iter, *ir)
 		{
 			ir_instruction *inst = (ir_instruction *)iter.get();
@@ -3610,7 +3761,7 @@ void ir_gen_glsl_visitor::AddTypeToUsedStructs(const glsl_type* type)
 char* FGlslCodeBackend::GenerateCode(exec_list* ir, _mesa_glsl_parse_state* state, EHlslShaderFrequency Frequency)
 {
 	FixRedundantCasts(ir);
-//IRDump(ir);
+	//IRDump(ir);
 
 	const bool bDefaultPrecisionIsHalf = ((HlslCompileFlags & HLSLCC_UseFullPrecisionInPS) == 0 && Frequency != HSF_ComputeShader);
 	const bool bUsesExternalTexture = ((HlslCompileFlags & HLSLCC_UsesExternalTexture) == HLSLCC_UsesExternalTexture);
@@ -3628,6 +3779,18 @@ char* FGlslCodeBackend::GenerateCode(exec_list* ir, _mesa_glsl_parse_state* stat
 	const bool bEmitPrecision = WantsPrecisionModifiers();
 	const bool bUsesFrameBufferFetch = Frequency == HSF_PixelShader && UsesUEIntrinsic(ir, FRAMEBUFFER_FETCH_ES2);
 	const bool bUsesDepthBufferFetch = Frequency == HSF_PixelShader && UsesUEIntrinsic(ir, DEPTHBUFFER_FETCH_ES2);
+
+	uint32 FramebufferFetchMask = 0;
+	uint32 FramebufferFetchWriteMask = 0;
+	if (Frequency == HSF_PixelShader)
+	{
+		for (int32 i = 0; i < UE_ARRAY_COUNT(GL_FRAMEBUFFER_FETCH); ++i)
+		{
+			FramebufferFetchWriteMask |= UsesUEIntrinsic(ir, GL_FRAMEBUFFER_FETCH_WRITE[i]) ? 1 << i : 0;
+			FramebufferFetchMask |= UsesUEIntrinsic(ir, GL_FRAMEBUFFER_FETCH[i]) ? 1 << i : 0;
+		}
+	}
+
 	ir_gen_glsl_visitor visitor(state->bGenerateES,
 								bEmitPrecision,
 								false,
@@ -3638,7 +3801,10 @@ char* FGlslCodeBackend::GenerateCode(exec_list* ir, _mesa_glsl_parse_state* stat
 								!AllowsGlobalUniforms(),
 								bUsesFrameBufferFetch,
 								bUsesDepthBufferFetch,
+								FramebufferFetchMask,
+								FramebufferFetchWriteMask,
 								bUsesExternalTexture);
+
 	const char* code = visitor.run(ir, state, bGroupFlattenedUBs);
 	return _strdup(code);
 }
@@ -5770,7 +5936,6 @@ void FGlslLanguageSpec::SetupLanguageIntrinsics(_mesa_glsl_parse_state* State, e
 		sig->is_builtin = true;
 		sig->is_defined = false;
 		sig->has_output_parameters = false;
-
 		func->add_signature(sig);
 		State->symbols->add_global_function(func);
 	}
@@ -5847,6 +6012,68 @@ void FGlslLanguageSpec::SetupLanguageIntrinsics(_mesa_glsl_parse_state* State, e
 			State->symbols->add_global_function(func);
 			ir->push_tail(func);
 		}
+
+		auto AddIntrisicReturningFloat = [](_mesa_glsl_parse_state* State, exec_list* ir, const char* Name, int32 NumColumns)
+		{
+			ir_function* Func = new(State) ir_function(Name);
+			auto* ReturnType = glsl_type::get_instance(GLSL_TYPE_FLOAT, NumColumns, 1);
+			ir_function_signature* Sig = new(State) ir_function_signature(ReturnType);
+			Sig->is_builtin = true;
+			Func->add_signature(Sig);
+			State->symbols->add_global_function(Func);
+			ir->push_head(Func);
+		};
+
+		for (int32 i = 0; i < UE_ARRAY_COUNT(GL_FRAMEBUFFER_FETCH); ++i)
+		{
+			if (!strcmp(GL_FRAMEBUFFER_FETCH_TYPE[i], "float"))
+			{
+				AddIntrisicReturningFloat(State, ir, GL_FRAMEBUFFER_FETCH[i], 1);
+			}
+			else if (!strcmp(GL_FRAMEBUFFER_FETCH_TYPE[i], "vec3"))
+			{
+				AddIntrisicReturningFloat(State, ir, GL_FRAMEBUFFER_FETCH[i], 3);
+			}
+			else
+			{
+				AddIntrisicReturningFloat(State, ir, GL_FRAMEBUFFER_FETCH[i], 4);
+			}
+		}
+
+		auto AddIntrisicVecParam = [](_mesa_glsl_parse_state* State, exec_list* ir, const char* Name, const glsl_type* GlslType)
+		{
+			ir_function* Func = new(State) ir_function(Name);
+
+			ir_variable* param = new (State) ir_variable(GlslType, "Src", ir_variable_mode::ir_var_in);
+			exec_list* params = new(State) exec_list();
+			params->push_tail(param);
+
+			ir_function_signature* sig = new(State)ir_function_signature(glsl_type::void_type);
+			sig->replace_parameters(params);
+			sig->is_defined = false;
+			sig->has_output_parameters = true;
+			sig->is_builtin = true;
+			
+			Func->add_signature(sig);
+			State->symbols->add_global_function(Func);
+			ir->push_head(Func);
+		};
+
+		for (int32 i = 0; i < UE_ARRAY_COUNT(GL_FRAMEBUFFER_FETCH_WRITE); ++i)
+		{
+			if (!strcmp(GL_FRAMEBUFFER_FETCH_TYPE[i], "float"))
+			{
+				AddIntrisicVecParam(State, ir, GL_FRAMEBUFFER_FETCH_WRITE[i], glsl_type::half_type);
+			}
+			else if (!strcmp(GL_FRAMEBUFFER_FETCH_TYPE[i], "vec3"))
+			{
+				AddIntrisicVecParam(State, ir, GL_FRAMEBUFFER_FETCH_WRITE[i], glsl_type::half3_type);
+			}
+			else
+			{
+				AddIntrisicVecParam(State, ir, GL_FRAMEBUFFER_FETCH_WRITE[i], glsl_type::half4_type);
+			}
+		}
 	}
 }
 
diff --git a/Engine/Source/Runtime/Apple/MetalRHI/Private/MetalRHI.cpp b/Engine/Source/Runtime/Apple/MetalRHI/Private/MetalRHI.cpp
--- a/Engine/Source/Runtime/Apple/MetalRHI/Private/MetalRHI.cpp
+++ b/Engine/Source/Runtime/Apple/MetalRHI/Private/MetalRHI.cpp
@@ -568,6 +568,7 @@ FMetalDynamicRHI::FMetalDynamicRHI(ERHIFeatureLevel::Type RequestedFeatureLevel)
 #endif
 
 	GSupportsShaderFramebufferFetch = !PLATFORM_MAC && GMaxRHIShaderPlatform != SP_METAL_MRT && GMaxRHIShaderPlatform != SP_METAL_MRT_TVOS;
+	GSupportsShaderMRTFramebufferFetch = GSupportsShaderFramebufferFetch;
 	GHardwareHiddenSurfaceRemoval = true;
 	GSupportsRenderTargetFormat_PF_G8 = false;
 	GRHISupportsTextureStreaming = true;
diff --git a/Engine/Source/Runtime/Engine/Private/ShaderCompiler/ShaderCompiler.cpp b/Engine/Source/Runtime/Engine/Private/ShaderCompiler/ShaderCompiler.cpp
--- a/Engine/Source/Runtime/Engine/Private/ShaderCompiler/ShaderCompiler.cpp
+++ b/Engine/Source/Runtime/Engine/Private/ShaderCompiler/ShaderCompiler.cpp
@@ -4115,6 +4115,9 @@ void GlobalBeginCompileShader(
 		if (IsMobileDeferredShadingEnabled((EShaderPlatform)Target.Platform))
 		{
 			Input.Environment.SetDefine(TEXT("MOBILE_DEFERRED_SHADING"), 1);
+
+			bool bGLESDeferredShading = Target.Platform == SP_OPENGL_ES3_1_ANDROID;
+			Input.Environment.SetDefine(TEXT("USE_GLES_FBF_DEFERRED"), bGLESDeferredShading ? 1 : 0);
 		}
 	}
 
diff --git a/Engine/Source/Runtime/OpenGLDrv/Private/OpenGLDevice.cpp b/Engine/Source/Runtime/OpenGLDrv/Private/OpenGLDevice.cpp
--- a/Engine/Source/Runtime/OpenGLDrv/Private/OpenGLDevice.cpp
+++ b/Engine/Source/Runtime/OpenGLDrv/Private/OpenGLDevice.cpp
@@ -849,7 +849,10 @@ static void InitRHICapabilitiesForGL()
 	GSupportsResourceView = FOpenGL::SupportsResourceView();
 
 	GSupportsShaderFramebufferFetch = FOpenGL::SupportsShaderFramebufferFetch();
+	GSupportsShaderMRTFramebufferFetch = FOpenGL::SupportsShaderMRTFramebufferFetch();
 	GSupportsShaderDepthStencilFetch = FOpenGL::SupportsShaderDepthStencilFetch();
+	GSupportsPixelLocalStorage = FOpenGL::SupportsPixelLocalStorage();
+
 	GMaxShadowDepthBufferSizeX = FMath::Min<int32>(Value_GL_MAX_RENDERBUFFER_SIZE, 4096); // Limit to the D3D11 max.
 	GMaxShadowDepthBufferSizeY = FMath::Min<int32>(Value_GL_MAX_RENDERBUFFER_SIZE, 4096);
 	GHardwareHiddenSurfaceRemoval = FOpenGL::HasHardwareHiddenSurfaceRemoval();
@@ -1087,10 +1090,10 @@ static bool VerifyCompiledShader(GLuint Shader, const ANSICHAR* GlslCode, bool I
 			for (int i = 0; i < 30 && (*Temp != '\0'); ++i)
 			{
 				FString Converted = ANSI_TO_TCHAR(Temp);
-				Converted.LeftChop(256);
+				FString TempString = Converted.LeftChop(256);
 
-				UE_LOG(LogRHI, Display, TEXT("%s"), *Converted);
-				Temp += Converted.Len();
+				UE_LOG(LogRHI, Display, TEXT("%s"), *TempString);
+				Temp += TempString.Len();
 			}
 		}
 #endif
diff --git a/Engine/Source/Runtime/OpenGLDrv/Private/OpenGLES.cpp b/Engine/Source/Runtime/OpenGLDrv/Private/OpenGLES.cpp
--- a/Engine/Source/Runtime/OpenGLDrv/Private/OpenGLES.cpp
+++ b/Engine/Source/Runtime/OpenGLDrv/Private/OpenGLES.cpp
@@ -84,6 +84,9 @@ bool FOpenGLES::bSupportsColorBufferFloat = false;
 /** GL_EXT_shader_framebuffer_fetch */
 bool FOpenGLES::bSupportsShaderFramebufferFetch = false;
 
+/** GL_EXT_shader_framebuffer_fetch (MRT's) */
+bool FOpenGLES::bSupportsShaderMRTFramebufferFetch = false;
+
 /* This is to avoid a bug where device supports GL_EXT_shader_framebuffer_fetch but does not define it in GLSL */
 bool FOpenGLES::bRequiresUEShaderFramebufferFetchDef = false;
 
@@ -99,6 +102,9 @@ bool FOpenGLES::bSupportsDXT = false;
 /** OpenGL ES 3.0 profile */
 bool FOpenGLES::bSupportsETC2 = false;
 
+/** GL_EXT_shader_pixel_local_storage */
+bool FOpenGLES::bSupportsPixelLocalStorage = false;
+
 /** GL_FRAGMENT_SHADER, GL_LOW_FLOAT */
 int FOpenGLES::ShaderLowPrecision = 0;
 
@@ -204,7 +210,9 @@ void FOpenGLES::ProcessExtensions(const FString& ExtensionsString)
 	bSupportsColorBufferHalfFloat = ExtensionsString.Contains(TEXT("GL_EXT_color_buffer_half_float"));
 	bSupportsShaderFramebufferFetch = ExtensionsString.Contains(TEXT("GL_EXT_shader_framebuffer_fetch")) || ExtensionsString.Contains(TEXT("GL_NV_shader_framebuffer_fetch"))
 		|| ExtensionsString.Contains(TEXT("GL_ARM_shader_framebuffer_fetch ")); // has space at the end to exclude GL_ARM_shader_framebuffer_fetch_depth_stencil match
+	bSupportsShaderMRTFramebufferFetch = ExtensionsString.Contains(TEXT("GL_EXT_shader_framebuffer_fetch")) || ExtensionsString.Contains(TEXT("GL_NV_shader_framebuffer_fetch"));
 	bRequiresUEShaderFramebufferFetchDef = ExtensionsString.Contains(TEXT("GL_EXT_shader_framebuffer_fetch"));
+	bSupportsPixelLocalStorage = ExtensionsString.Contains(TEXT("GL_EXT_shader_pixel_local_storage"));
 	bSupportsShaderDepthStencilFetch = ExtensionsString.Contains(TEXT("GL_ARM_shader_framebuffer_fetch_depth_stencil"));
 	bSupportsMultisampledRenderToTexture = ExtensionsString.Contains(TEXT("GL_EXT_multisampled_render_to_texture"));
 	bSupportsDXT = ExtensionsString.Contains(TEXT("GL_NV_texture_compression_s3tc")) || ExtensionsString.Contains(TEXT("GL_EXT_texture_compression_s3tc"));
diff --git a/Engine/Source/Runtime/OpenGLDrv/Private/OpenGLRenderTarget.cpp b/Engine/Source/Runtime/OpenGLDrv/Private/OpenGLRenderTarget.cpp
--- a/Engine/Source/Runtime/OpenGLDrv/Private/OpenGLRenderTarget.cpp
+++ b/Engine/Source/Runtime/OpenGLDrv/Private/OpenGLRenderTarget.cpp
@@ -1166,6 +1166,15 @@ void FOpenGLDynamicRHI::RHIBeginRenderPass(const FRHIRenderPassInfo& InInfo, con
 		BeginOcclusionQueryBatch(InInfo.NumOcclusionQueries);
 	}
 
+#if PLATFORM_ANDROID
+	if (RenderPassInfo.SubpassHint == ESubpassHint::DeferredShadingSubpass &&
+		!FOpenGL::SupportsShaderMRTFramebufferFetch() &&
+		FOpenGL::SupportsPixelLocalStorage())
+	{
+		glEnable(GL_SHADER_PIXEL_LOCAL_STORAGE_EXT);
+	}
+#endif
+
 #if PLATFORM_ANDROID && !PLATFORM_LUMIN && !PLATFORM_LUMINGL4
 	if (FAndroidOpenGL::RequiresAdrenoTilingModeHint())
 	{
@@ -1215,6 +1224,15 @@ void FOpenGLDynamicRHI::RHIEndRenderPass()
 	FRHIRenderTargetView RTV(nullptr, ERenderTargetLoadAction::ENoAction);
 	FRHIDepthRenderTargetView DepthRTV(nullptr, ERenderTargetLoadAction::ENoAction, ERenderTargetStoreAction::ENoAction);
 	SetRenderTargets(1, &RTV, &DepthRTV);
+
+#if PLATFORM_ANDROID && !PLATFORM_LUMIN && !PLATFORM_LUMINGL4
+	if (RenderPassInfo.SubpassHint == ESubpassHint::DeferredShadingSubpass &&
+		!FOpenGL::SupportsShaderMRTFramebufferFetch() &&
+		FOpenGL::SupportsPixelLocalStorage())
+	{
+		glDisable(GL_SHADER_PIXEL_LOCAL_STORAGE_EXT);
+	}
+#endif
 }
 
 void FOpenGLDynamicRHI::RHINextSubpass()
diff --git a/Engine/Source/Runtime/OpenGLDrv/Private/OpenGLShaders.cpp b/Engine/Source/Runtime/OpenGLDrv/Private/OpenGLShaders.cpp
--- a/Engine/Source/Runtime/OpenGLDrv/Private/OpenGLShaders.cpp
+++ b/Engine/Source/Runtime/OpenGLDrv/Private/OpenGLShaders.cpp
@@ -888,11 +888,22 @@ void OPENGLDRV_API GLSLToDeviceCompatibleGLSL(FAnsiCharArray& GlslCodeOriginal,
 
 	// The incoming glsl may have preprocessor code that is dependent on defines introduced via the engine.
 	// This is the place to insert such engine preprocessor defines, immediately after the glsl version declaration.
-	if (Capabilities.bRequiresUEShaderFramebufferFetchDef && TypeEnum == GL_FRAGMENT_SHADER)
+	if (TypeEnum == GL_FRAGMENT_SHADER)
 	{
-		// Some devices (Zenfone5) support GL_EXT_shader_framebuffer_fetch but do not define GL_EXT_shader_framebuffer_fetch in GLSL compiler
-		// We can't define anything with GL_, so we use UE_EXT_shader_framebuffer_fetch to enable frame buffer fetch
-		AppendCString(GlslCode, "#define UE_EXT_shader_framebuffer_fetch 1\n");
+		if (Capabilities.bRequiresUEShaderFramebufferFetchDef)
+		{
+			// Some devices (Zenfone5) support GL_EXT_shader_framebuffer_fetch but do not define GL_EXT_shader_framebuffer_fetch in GLSL compiler
+			// We can't define anything with GL_, so we use UE_EXT_shader_framebuffer_fetch to enable frame buffer fetch
+			AppendCString(GlslCode, "#define UE_EXT_shader_framebuffer_fetch 1\n");
+		}
+		if(FOpenGL::SupportsShaderMRTFramebufferFetch())
+		{
+			AppendCString(GlslCode, "#define UE_MRT_FRAMEBUFFER_FETCH 1\n");
+		}
+		else if (FOpenGL::SupportsPixelLocalStorage())
+		{
+			AppendCString(GlslCode, "#define UE_MRT_PLS 1\n");
+		}
 	}
 
 	if (bEmitTextureExternal)
diff --git a/Engine/Source/Runtime/OpenGLDrv/Public/OpenGL.h b/Engine/Source/Runtime/OpenGLDrv/Public/OpenGL.h
--- a/Engine/Source/Runtime/OpenGLDrv/Public/OpenGL.h
+++ b/Engine/Source/Runtime/OpenGLDrv/Public/OpenGL.h
@@ -110,7 +110,9 @@ public:
 	static FORCEINLINE bool SupportsColorBufferHalfFloat()				{ return true; }
 	static FORCEINLINE bool SupportsVolumeTextureRendering()			{ return false; }
 	static FORCEINLINE bool SupportsShaderFramebufferFetch()			{ return false; }
+	static FORCEINLINE bool SupportsShaderMRTFramebufferFetch()			{ return false; }
 	static FORCEINLINE bool SupportsShaderDepthStencilFetch()			{ return false; }
+	static FORCEINLINE bool SupportsPixelLocalStorage()					{ return false; }
 	static FORCEINLINE bool SupportsVertexArrayBGRA()					{ return true; }
 	static FORCEINLINE bool SupportsBGRA8888()							{ return true; }
 	static FORCEINLINE bool SupportsDXT()								{ return true; }
diff --git a/Engine/Source/Runtime/OpenGLDrv/Public/OpenGLES.h b/Engine/Source/Runtime/OpenGLDrv/Public/OpenGLES.h
--- a/Engine/Source/Runtime/OpenGLDrv/Public/OpenGLES.h
+++ b/Engine/Source/Runtime/OpenGLDrv/Public/OpenGLES.h
@@ -202,7 +202,9 @@ struct FOpenGLES : public FOpenGLBase
 	static FORCEINLINE bool SupportsColorBufferFloat() { return bSupportsColorBufferFloat; }
 	static FORCEINLINE bool SupportsColorBufferHalfFloat() { return bSupportsColorBufferHalfFloat; }
 	static FORCEINLINE bool SupportsShaderFramebufferFetch() { return bSupportsShaderFramebufferFetch; }
+	static FORCEINLINE bool SupportsShaderMRTFramebufferFetch() { return bSupportsShaderMRTFramebufferFetch; }
 	static FORCEINLINE bool SupportsShaderDepthStencilFetch() { return bSupportsShaderDepthStencilFetch; }
+	static FORCEINLINE bool SupportsPixelLocalStorage() { return bSupportsPixelLocalStorage; }
 	static FORCEINLINE bool SupportsMultisampledRenderToTexture() { return bSupportsMultisampledRenderToTexture; }
 	static FORCEINLINE bool SupportsVertexArrayBGRA() { return false; }
 	static FORCEINLINE bool SupportsBGRA8888() { return bSupportsBGRA8888; }
@@ -767,6 +769,9 @@ protected:
 	/** GL_EXT_shader_framebuffer_fetch */
 	static bool bSupportsShaderFramebufferFetch;
 
+	/** GL_EXT_shader_framebuffer_fetch (MRT's) */
+	static bool bSupportsShaderMRTFramebufferFetch;
+
 	/** workaround for GL_EXT_shader_framebuffer_fetch */
 	static bool bRequiresUEShaderFramebufferFetchDef;
 
@@ -776,6 +781,9 @@ protected:
 	/** GL_EXT_MULTISAMPLED_RENDER_TO_TEXTURE */
 	static bool bSupportsMultisampledRenderToTexture;
 
+	/** workaround for GL_EXT_shader_pixel_local_storage */
+	static bool bSupportsPixelLocalStorage;
+
 	/** GL_FRAGMENT_SHADER, GL_LOW_FLOAT */
 	static int ShaderLowPrecision;
 
diff --git a/Engine/Source/Runtime/RHI/Private/RHI.cpp b/Engine/Source/Runtime/RHI/Private/RHI.cpp
--- a/Engine/Source/Runtime/RHI/Private/RHI.cpp
+++ b/Engine/Source/Runtime/RHI/Private/RHI.cpp
@@ -781,6 +781,8 @@ TRHIGlobal<bool> GSupportsRenderTargetFormat_PF_G8(true);
 TRHIGlobal<bool> GSupportsRenderTargetFormat_PF_FloatRGBA(true);
 bool GSupportsShaderFramebufferFetch = false;
 bool GSupportsShaderDepthStencilFetch = false;
+bool GSupportsShaderMRTFramebufferFetch = false;
+bool GSupportsPixelLocalStorage = false;
 bool GSupportsTimestampRenderQueries = false;
 bool GRHISupportsGPUTimestampBubblesRemoval = false;
 bool GRHISupportsFrameCyclesBubblesRemoval = false;
diff --git a/Engine/Source/Runtime/RHI/Public/RHI.h b/Engine/Source/Runtime/RHI/Public/RHI.h
--- a/Engine/Source/Runtime/RHI/Public/RHI.h
+++ b/Engine/Source/Runtime/RHI/Public/RHI.h
@@ -163,7 +163,7 @@ inline bool RHISupportsMSAA(const FStaticShaderPlatform Platform)
 
 inline bool RHISupportsBufferLoadTypeConversion(const FStaticShaderPlatform Platform)
 {
-	return !IsMetalPlatform(Platform);
+	return !IsMetalPlatform(Platform) && !IsOpenGLPlatform(Platform);
 }
 
 /** Whether the platform supports reading from volume textures (does not cover rendering to volume textures). */
@@ -305,6 +305,12 @@ extern RHI_API TRHIGlobal<bool> GSupportsRenderTargetFormat_PF_FloatRGBA;
 /** true if mobile framebuffer fetch is supported */
 extern RHI_API bool GSupportsShaderFramebufferFetch;
 
+/** true if mobile framebuffer fetch is supported from MRT's*/
+extern RHI_API bool GSupportsShaderMRTFramebufferFetch;
+
+/** true if mobile pixel local storage is supported */
+extern RHI_API bool GSupportsPixelLocalStorage;
+
 /** true if mobile depth & stencil fetch is supported */
 extern RHI_API bool GSupportsShaderDepthStencilFetch;
 
diff --git a/Engine/Source/Runtime/RenderCore/Private/RenderUtils.cpp b/Engine/Source/Runtime/RenderCore/Private/RenderUtils.cpp
--- a/Engine/Source/Runtime/RenderCore/Private/RenderUtils.cpp
+++ b/Engine/Source/Runtime/RenderCore/Private/RenderUtils.cpp
@@ -977,11 +977,6 @@ RENDERCORE_API bool MobileSupportsGPUScene(const FStaticShaderPlatform Platform)
 
 RENDERCORE_API bool IsMobileDeferredShadingEnabled(const FStaticShaderPlatform Platform)
 {
-	if (IsOpenGLPlatform(Platform))
-	{
-		// needs MRT framebuffer fetch or PLS
-		return false;
-	}
 	static auto* MobileShadingPathCvar = IConsoleManager::Get().FindTConsoleVariableDataInt(TEXT("r.Mobile.ShadingPath"));
 	return MobileShadingPathCvar->GetValueOnAnyThread() == 1;
 }
diff --git a/Engine/Source/Runtime/Renderer/Private/MobileShadingRenderer.cpp b/Engine/Source/Runtime/Renderer/Private/MobileShadingRenderer.cpp
--- a/Engine/Source/Runtime/Renderer/Private/MobileShadingRenderer.cpp
+++ b/Engine/Source/Runtime/Renderer/Private/MobileShadingRenderer.cpp
@@ -392,7 +392,8 @@ void FMobileSceneRenderer::InitViews(FRDGBuilder& GraphBuilder, FSceneTexturesCo
 		bRequiresPixelProjectedPlanarRelfectionPass ||
 		bSeparateTranslucencyActive ||
 		Views[0].bIsReflectionCapture ||
-		(bDeferredShading && bPostProcessUsesSceneDepth);
+		(bDeferredShading && bPostProcessUsesSceneDepth) ||
+		IsAndroidOpenGLESPlatform(ShaderPlatform);
 	// never keep MSAA depth
 	bKeepDepthContent = (NumMSAASamples > 1 ? false : bKeepDepthContent);
 
@@ -1008,19 +1009,116 @@ void FMobileSceneRenderer::RenderForward(FRDGBuilder& GraphBuilder, const TArray
 	QueueSceneTextureExtractions(GraphBuilder, SceneTextures);
 }
 
+class FMobileDeferredCopyPLSPS : public FGlobalShader
+{
+	DECLARE_SHADER_TYPE(FMobileDeferredCopyPLSPS, Global);
+
+	static void ModifyCompilationEnvironment(const FGlobalShaderPermutationParameters& Parameters, FShaderCompilerEnvironment& OutEnvironment)
+	{
+		FGlobalShader::ModifyCompilationEnvironment(Parameters, OutEnvironment);
+	}
+
+	static bool ShouldCompilePermutation(const FGlobalShaderPermutationParameters& Parameters)
+	{
+		return IsMobilePlatform(Parameters.Platform) && IsMobileDeferredShadingEnabled(Parameters.Platform);
+	}
+
+	/** Default constructor. */
+	FMobileDeferredCopyPLSPS() {}
+
+	/** Initialization constructor. */
+	FMobileDeferredCopyPLSPS(const ShaderMetaType::CompiledShaderInitializerType& Initializer)
+		: FGlobalShader(Initializer)
+	{
+	}
+};
+
+IMPLEMENT_SHADER_TYPE(, FMobileDeferredCopyPLSPS, TEXT("/Engine/Private/MobileDeferredUtils.usf"), TEXT("MobileDeferredCopyPLSPS"), SF_Pixel);
+
+class FMobileDeferredCopyDepthPS : public FGlobalShader
+{
+	DECLARE_SHADER_TYPE(FMobileDeferredCopyDepthPS, Global);
+
+	static void ModifyCompilationEnvironment(const FGlobalShaderPermutationParameters& Parameters, FShaderCompilerEnvironment& OutEnvironment)
+	{
+		FGlobalShader::ModifyCompilationEnvironment(Parameters, OutEnvironment);
+	}
+
+	static bool ShouldCompilePermutation(const FGlobalShaderPermutationParameters& Parameters)
+	{
+		return IsMobilePlatform(Parameters.Platform) && IsMobileDeferredShadingEnabled(Parameters.Platform);
+	}
+
+	/** Default constructor. */
+	FMobileDeferredCopyDepthPS() {}
+
+	/** Initialization constructor. */
+	FMobileDeferredCopyDepthPS(const ShaderMetaType::CompiledShaderInitializerType& Initializer)
+		: FGlobalShader(Initializer)
+	{
+	}
+};
+
+IMPLEMENT_SHADER_TYPE(, FMobileDeferredCopyDepthPS, TEXT("/Engine/Private/MobileDeferredUtils.usf"), TEXT("MobileDeferredCopyDepthPS"), SF_Pixel);
+
+template<class T>
+void MobileDeferredCopyBuffer(FRHICommandListImmediate& RHICmdList, const FViewInfo& View)
+{
+	FGraphicsPipelineStateInitializer GraphicsPSOInit;
+	RHICmdList.ApplyCachedRenderTargets(GraphicsPSOInit);
+	GraphicsPSOInit.BlendState = TStaticBlendState<>::GetRHI();
+	GraphicsPSOInit.RasterizerState = TStaticRasterizerState<>::GetRHI();
+	// Shade only MSM_DefaultLit pixels
+	uint8 StencilRef = GET_STENCIL_MOBILE_SM_MASK(MSM_DefaultLit);
+	GraphicsPSOInit.DepthStencilState = TStaticDepthStencilState<false, CF_Always>::GetRHI(); // 4 bits for shading models
+
+	TShaderMapRef<FPostProcessVS> VertexShader(View.ShaderMap);
+	TShaderMapRef<T> PixelShader(View.ShaderMap);
+
+	GraphicsPSOInit.BoundShaderState.VertexDeclarationRHI = GFilterVertexDeclaration.VertexDeclarationRHI;
+	GraphicsPSOInit.BoundShaderState.VertexShaderRHI = VertexShader.GetVertexShader();
+	GraphicsPSOInit.BoundShaderState.PixelShaderRHI = PixelShader.GetPixelShader();
+	GraphicsPSOInit.PrimitiveType = PT_TriangleList;
+	SetGraphicsPipelineState(RHICmdList, GraphicsPSOInit);
+	RHICmdList.SetStencilRef(StencilRef);
+
+	DrawRectangle(
+		RHICmdList,
+		0, 0,
+		View.ViewRect.Width(), View.ViewRect.Height(),
+		View.ViewRect.Min.X, View.ViewRect.Min.Y,
+		View.ViewRect.Width(), View.ViewRect.Height(),
+		FIntPoint(View.ViewRect.Width(), View.ViewRect.Height()),
+		FSceneTexturesConfig::Get().Extent,
+		VertexShader);
+}
+
 void FMobileSceneRenderer::RenderDeferred(FRDGBuilder& GraphBuilder, const TArrayView<const FViewInfo> ViewList, const FSortedLightSetSceneInfo& SortedLightSet, FRDGTextureRef ViewFamilyTexture, const FSceneTextures& SceneTextures)
 {
-	FRDGTextureRef ColorTargets[5] = {
-		SceneTextures.Color.Target,
-		SceneTextures.GBufferA,
-		SceneTextures.GBufferB,
-		SceneTextures.GBufferC,
-		SceneTextures.DepthAux
-	};
+	TArray<FRDGTextureRef, TInlineAllocator<5>> ColorTargets;
+
+	// If we are using GL and don't have FBF support, use PLS
+	bool bUsingPixelLocalStorage = IsAndroidOpenGLESPlatform(ShaderPlatform) && GSupportsPixelLocalStorage && !GSupportsShaderMRTFramebufferFetch;
+
+	if (bUsingPixelLocalStorage)
+	{
+		ColorTargets.Add(SceneTextures.Color.Target);
+	}
+	else
+	{
+		ColorTargets.Add(SceneTextures.Color.Target);
+		ColorTargets.Add(SceneTextures.GBufferA);
+		ColorTargets.Add(SceneTextures.GBufferB);
+		ColorTargets.Add(SceneTextures.GBufferC);
+		if(MobileRequiresSceneDepthAux(ShaderPlatform))
+		{
+			ColorTargets.Add(SceneTextures.DepthAux);
+		}
+	}
 
-	TArrayView<FRDGTextureRef> BasePassTexturesView = MakeArrayView(ColorTargets, MobileRequiresSceneDepthAux(ShaderPlatform) ? 5 : 4);
+	TArrayView<FRDGTextureRef> BasePassTexturesView = MakeArrayView(ColorTargets);
 
-	FRenderTargetBindingSlots BasePassRenderTargets = GetRenderTargetBindings(ERenderTargetLoadAction::EClear, BasePassTexturesView);
+	FRenderTargetBindingSlots BasePassRenderTargets = GetRenderTargetBindings(ERenderTargetLoadAction::ENoAction, BasePassTexturesView);
 	BasePassRenderTargets.DepthStencil = FDepthStencilBinding(SceneTextures.Depth.Target, ERenderTargetLoadAction::EClear, FExclusiveDepthStencil::DepthWrite_StencilWrite);
 	BasePassRenderTargets.SubpassHint = ESubpassHint::DeferredShadingSubpass;
 	BasePassRenderTargets.NumOcclusionQueries = ComputeNumOcclusionQueriesToBatch();
@@ -1086,6 +1184,27 @@ void FMobileSceneRenderer::RenderDeferred(FRDGBuilder& GraphBuilder, const TArra
 			// SceneColor write, SceneDepth is read only
 			RHICmdList.NextSubpass();
 		});
+
+		if (bUsingPixelLocalStorage)
+		{
+			{
+				auto* PassParameters = GraphBuilder.AllocParameters<FRenderTargetParameters>();
+				PassParameters->RenderTargets = BasePassRenderTargets;
+
+				AddPass(GraphBuilder, [this](FRHICommandListImmediate& RHICmdList)
+				{
+					MobileDeferredCopyBuffer<FMobileDeferredCopyPLSPS>(RHICmdList, Views[0]);
+				});
+			}
+
+			{
+				AddPass(GraphBuilder, [](FRHICommandListImmediate& RHICmdList)
+				{
+					// SceneColor write, SceneDepth is read only
+					RHICmdList.NextSubpass();
+				});
+			}
+		}
 		
 		MobileDeferredShadingPass(GraphBuilder, BasePassRenderTargets, SceneTextures.MobileUniformBuffer, *Scene, ViewList, SortedLightSet);
 		// Draw translucency.
@@ -1097,6 +1216,31 @@ void FMobileSceneRenderer::RenderDeferred(FRDGBuilder& GraphBuilder, const TArra
 			AddPass(GraphBuilder, PollOcclusionQueriesAndDispatchToRHIThreadPass);
 		}
 
+		/*if (IsAndroidOpenGLESPlatform(ShaderPlatform))
+		{
+			FRDGTextureRef DepthCopyTargets[] = 
+			{
+				SceneTextures.Depth.Target,
+			};
+			
+			TArrayView<FRDGTextureRef> DepthCopyTargetsView = MakeArrayView(DepthCopyTargets);
+
+			FRenderTargetBindingSlots DepthCopyRenderTargets = GetRenderTargetBindings(ERenderTargetLoadAction::EClear, DepthCopyTargetsView);
+			DepthCopyRenderTargets.DepthStencil = FDepthStencilBinding(SceneTextures.Depth.Target, ERenderTargetLoadAction::ELoad, FExclusiveDepthStencil::DepthRead_StencilRead);
+			DepthCopyRenderTargets.NumOcclusionQueries = 0;
+			DepthCopyRenderTargets.FoveationTexture = nullptr;
+			DepthCopyRenderTargets.MultiViewCount = 0;
+
+			auto* PassParameters = GraphBuilder.AllocParameters<FRenderTargetParameters>();
+			PassParameters->RenderTargets = DepthCopyRenderTargets;
+
+			GraphBuilder.AddPass(RDG_EVENT_NAME("FBFDepthCopy"), PassParameters, ERDGPassFlags::Raster,
+				[this](FRHICommandListImmediate& RHICmdList)
+			{
+				MobileDeferredCopyBuffer<FMobileDeferredCopyDepthPS>(RHICmdList, Views[0]);
+			});
+		}*/
+
 		AddPass(GraphBuilder, [](FRHICommandListImmediate& RHICmdList)
 		{
 			RHICmdList.EndRenderPass();
@@ -1113,7 +1257,7 @@ void FMobileSceneRenderer::RenderDeferred(FRDGBuilder& GraphBuilder, const TArra
 		
 		// SceneColor + GBuffer write, SceneDepth is read only
 		{
-			for (uint32 i = 0; i < UE_ARRAY_COUNT(ColorTargets); ++i)
+			for (int32 i = 0; i < ColorTargets.Num(); ++i)
 			{
 				BasePassRenderTargets[i].SetLoadAction(ERenderTargetLoadAction::ELoad);
 			}
@@ -1150,7 +1294,7 @@ void FMobileSceneRenderer::RenderDeferred(FRDGBuilder& GraphBuilder, const TArra
 
 		// SceneColor write, SceneDepth is read only
 		{
-			for (uint32 i = 1; i < UE_ARRAY_COUNT(ColorTargets); ++i)
+			for (int32 i = 1; i < ColorTargets.Num(); ++i)
 			{
 				BasePassRenderTargets[i] = FRenderTargetBinding();
 			}
@@ -1258,17 +1402,17 @@ bool FMobileSceneRenderer::RequiresMultiPass(FRHICommandListImmediate& RHICmdLis
 		return false;
 	}
 
-	if (IsMobileDeferredShadingEnabled(ShaderPlatform))
-	{
-		// TODO: add GL support
-		return true;
-	}
-	
 	// Some Androids support frame_buffer_fetch
 	if (IsAndroidOpenGLESPlatform(ShaderPlatform) && (GSupportsShaderFramebufferFetch || GSupportsShaderDepthStencilFetch))
 	{
 		return false;
 	}
+
+	if (IsMobileDeferredShadingEnabled(ShaderPlatform))
+	{
+		// TODO: add GL support
+		return true;
+	}
 		
 	// Always render reflection capture in single pass
 	if (View.bIsPlanarReflection || View.bIsSceneCapture)
