Use GL_ARM_shader_framebuffer_fetch_depth_stencil by default for depth fetch on Android
This works better than full FBF and supported on most Adnroid devices. Fallback to FBF when extension is not supported
(merged from  Release-Engine-Staging)

diff --git a/Engine/Config/BaseDeviceProfiles.ini b/Engine/Config/BaseDeviceProfiles.ini
--- a/Engine/Config/BaseDeviceProfiles.ini
+++ b/Engine/Config/BaseDeviceProfiles.ini
@@ -712,14 +712,10 @@ BaseProfileName=Android_Low
 [Android_Adreno5xx_Low DeviceProfile]
 DeviceType=Android
 BaseProfileName=Android_Low
-; disable QCOM_shader_framebuffer_fetch_noncoherent on Adreno5xx, as it causes rendering artifacts with enabled MSAA
-+CVars=r.Android.DisableFBFNonCoherent=1
 
 [Android_Adreno5xx DeviceProfile]
 DeviceType=Android
 BaseProfileName=Android_Mid
-; disable QCOM_shader_framebuffer_fetch_noncoherent on Adreno5xx, as it causes rendering artifacts with enabled MSAA
-+CVars=r.Android.DisableFBFNonCoherent=1
 
 [Android_Adreno6xx DeviceProfile]
 DeviceType=Android
diff --git a/Engine/Shaders/Private/SceneTexturesCommon.ush b/Engine/Shaders/Private/SceneTexturesCommon.ush
--- a/Engine/Shaders/Private/SceneTexturesCommon.ush
+++ b/Engine/Shaders/Private/SceneTexturesCommon.ush
@@ -142,13 +142,7 @@ float LookupDeviceZ( float2 ScreenUV )
 		return Texture2DSample(MobileSceneTextures.SceneColorTexture, MobileSceneTextures.SceneColorTextureSampler, ScreenUV).a;
 	#endif
 #elif COMPILER_GLSL_ES3_1 && PIXELSHADER
-	#if !OUTPUT_MOBILE_HDR
-		// Try framebuffer_fetch_depth_stencil to get the depth directly if the extension is available.
-		// We cannot fall back to fetching the alpha channel when MobileHDR=false because the alpha channel is only 8-bit.
-		return DepthbufferFetchES2();
-	#else
-		return FramebufferFetchES2().w;
-	#endif
+	return DepthbufferFetchES2();
 #elif VULKAN_SUBPASS_DEPTHFETCH && PIXELSHADER
 	// Special intrinsic to read from the current depth buffer
 	return VulkanSubpassDepthFetch();
diff --git a/Engine/Source/Developer/ShaderFormatOpenGL/Private/GlslBackend.cpp b/Engine/Source/Developer/ShaderFormatOpenGL/Private/GlslBackend.cpp
--- a/Engine/Source/Developer/ShaderFormatOpenGL/Private/GlslBackend.cpp
+++ b/Engine/Source/Developer/ShaderFormatOpenGL/Private/GlslBackend.cpp
@@ -371,7 +371,7 @@ const char* GL_FRAMEBUFFER_FETCH_TYPE[5] =
 
 static const uint32_t FRAMEBUFFER_FETCH_DEPTH_INDEX = 4;
 
-static const char* ES31FrameBufferFetchStorageQualifier = "FRAME_BUFFERFETCH_STORAGE_QUALIFIER ";
+static const char* FBF_StorageQualifier = "FBF_STORAGE_QUALIFIER ";
 static_assert((sizeof(GLSLExpressionTable) / sizeof(GLSLExpressionTable[0])) == ir_opcode_count, "GLSLExpressionTableSizeMismatch");
 
 struct SDMARange
@@ -628,18 +628,12 @@ class ir_gen_glsl_visitor : public ir_visitor
 	/** number of loops in the generated code */
 	int loop_count;
 
-	/** Whether the shader being cross compiled needs EXT_shader_texture_lod. */
-	bool bUsesGLTextureLODExtension;
-
 	/** Whether the shader being cross compiled needs GL_EXT_texture_buffer. */
 	bool bUsesTextureBuffer;
 
 	/** Whether the shader being cross compiled needs GL_OES_shader_image_atomic. */
 	bool bUseImageAtomic;
 
-	// Found dFdx or dFdy
-	bool bUsesDXDY;
-
 	// True if the discard instruction was encountered.
 	bool bUsesDiscard;
 
@@ -1252,9 +1246,9 @@ class ir_gen_glsl_visitor : public ir_visitor
 					check(layout_bits == 0);
 					layout = ralloc_asprintf(nullptr, "layout(location=%d) ", var->location);
 				}
-
-				const bool bOverrideFBFOutputVarStorageQualifier = bIsES31 && bUsesFrameBufferFetch && var->name && (strncmp(var->name, "out_Target0", 11) == 0);
-				const char* StorageQualifier = bOverrideFBFOutputVarStorageQualifier ? ES31FrameBufferFetchStorageQualifier : mode_str[var->mode];
+				
+				const bool bNeedsFBFOutput = (bUsesFrameBufferFetch || bUsesDepthbufferFetch) && var->name && (strncmp(var->name, "out_Target0", 11) == 0);
+				const char* StorageQualifier = bNeedsFBFOutput ? FBF_StorageQualifier : mode_str[var->mode];
 
 				ralloc_asprintf_append(
 					buffer,
@@ -1334,6 +1328,32 @@ class ir_gen_glsl_visitor : public ir_visitor
 
 	virtual void visit(ir_function_signature *sig)
 	{
+		if (sig->is_main && ShaderTarget == fragment_shader)
+		{
+			// print this right before 'main' where all outputs are already declared
+			if (bUsesFrameBufferFetch)
+			{
+				ralloc_asprintf_append(buffer, "\n#ifdef GL_EXT_shader_framebuffer_fetch\n");
+				ralloc_asprintf_append(buffer, "  vec4 FramebufferFetchES2() { return out_Target0; }\n");
+				ralloc_asprintf_append(buffer, "#elif defined(GL_ARM_shader_framebuffer_fetch)\n");
+				ralloc_asprintf_append(buffer, "  vec4 FramebufferFetchES2() { return gl_LastFragColorARM; }\n");
+				ralloc_asprintf_append(buffer, "#else\n");
+				ralloc_asprintf_append(buffer, "  vec4 FramebufferFetchES2() { return vec4(0.0, 0.0, 0.0, 0.0); }\n");
+				ralloc_asprintf_append(buffer, "#endif\n\n");
+			}
+			
+			if (bUsesDepthbufferFetch)
+			{
+				ralloc_asprintf_append(buffer, "\n#ifdef GL_ARM_shader_framebuffer_fetch_depth_stencil\n");
+				ralloc_asprintf_append(buffer, "  float DepthbufferFetchES2() { return gl_LastFragDepthARM; }\n");
+				ralloc_asprintf_append(buffer, "#elif defined(GL_EXT_shader_framebuffer_fetch)\n");
+				ralloc_asprintf_append(buffer, "  float DepthbufferFetchES2() { return out_Target0.w; }\n");
+				ralloc_asprintf_append(buffer, "#else\n");
+				ralloc_asprintf_append(buffer, "  float DepthbufferFetchES2() { return 0.0; }\n");
+				ralloc_asprintf_append(buffer, "#endif\n\n");
+			}
+		}
+				
 		// Reset temporary id count.
 		temp_id = 0;
 		bool bPrintComma = false;
@@ -1473,11 +1493,6 @@ class ir_gen_glsl_visitor : public ir_visitor
 		}
 		else if (numOps < 4)
 		{
-			if (op == ir_unop_dFdx || op == ir_unop_dFdy)
-			{
-				bUsesDXDY = true;
-			}
-
 			ralloc_asprintf_append(buffer, GLSLExpressionTable[op][0]);
 			for (int i = 0; i < numOps; ++i)
 			{
@@ -1511,15 +1526,6 @@ class ir_gen_glsl_visitor : public ir_visitor
 			op = ir_tex;
 		}
 
-		bool bEmitEXT = false;
-
-		if (bIsES && op == ir_txl && ShaderTarget == fragment_shader)
-		{
-			// See http://www.khronos.org/registry/gles/extensions/EXT/EXT_shader_texture_lod.txt
-			bUsesGLTextureLODExtension = true;
-			bEmitEXT = true;
-		}
-
 		if (op == ir_txf)
 		{
 			bUsesTextureBuffer = true;
@@ -1536,7 +1542,7 @@ class ir_gen_glsl_visitor : public ir_visitor
 			grad_str[op == ir_txd],
 			lod_str[op == ir_txl],
 			offset_str[tex->offset != 0],
-			EXT_str[(int)bEmitEXT]
+			EXT_str[0]
 		);
 		tex->sampler->accept(this);
 
@@ -3171,39 +3177,18 @@ class ir_gen_glsl_visitor : public ir_visitor
 			ralloc_asprintf_append(buffer, "// Uses samplerExternalOES\n");
 		}
 
-		if (bUsesGLTextureLODExtension)
-		{
-			ralloc_asprintf_append(buffer, "#ifndef DONTEMITEXTENSIONSHADERTEXTURELODENABLE\n");
-			ralloc_asprintf_append(buffer, "#extension GL_EXT_shader_texture_lod : enable\n");
-			ralloc_asprintf_append(buffer, "#endif\n");
-		}
-
 		if (state->bSeparateShaderObjects && !state->bGenerateES && 
 			((state->target == tessellation_control_shader) || (state->target == tessellation_evaluation_shader)))
 		{
 			ralloc_asprintf_append(buffer, "#extension GL_ARB_tessellation_shader : enable\n");
 		}
 
-		if (bUsesDXDY && bIsES)
-		{
-			ralloc_asprintf_append(buffer, "#extension GL_OES_standard_derivatives : enable\n");
-		}
-
-		if (bUsesInstanceID && bIsES)
-		{
-			ralloc_asprintf_append(buffer, "#ifdef UE_EXT_draw_instanced\n");
-			ralloc_asprintf_append(buffer, "#extension GL_EXT_draw_instanced : enable\n");
-			ralloc_asprintf_append(buffer, "#define gl_InstanceID gl_InstanceIDEXT\n");
-			ralloc_asprintf_append(buffer, "#endif\n");
-		}
-		
 		if (bUsesFramebufferFetchES2)
 		{
-			ralloc_asprintf_append(buffer, "\n#ifdef UE_EXT_shader_framebuffer_fetch\n");
+			ralloc_asprintf_append(buffer, "\n#ifdef GL_EXT_shader_framebuffer_fetch\n");
 			ralloc_asprintf_append(buffer, "#extension GL_EXT_shader_framebuffer_fetch : enable\n");
-			ralloc_asprintf_append(buffer, "#define EXT_shader_framebuffer_fetch_enabled 1\n");
-			ralloc_asprintf_append(buffer, "#endif\n");
-			ralloc_asprintf_append(buffer, "\n#ifdef GL_ARM_shader_framebuffer_fetch\n");
+			ralloc_asprintf_append(buffer, "#define %sinout\n", FBF_StorageQualifier);
+			ralloc_asprintf_append(buffer, "#elif defined(GL_ARM_shader_framebuffer_fetch)\n");
 			ralloc_asprintf_append(buffer, "#extension GL_ARM_shader_framebuffer_fetch : enable\n");
 			ralloc_asprintf_append(buffer, "#endif\n");
 		}
@@ -3212,10 +3197,22 @@ class ir_gen_glsl_visitor : public ir_visitor
 			ralloc_asprintf_append(buffer, "#ifdef UE_MRT_FRAMEBUFFER_FETCH\n");
 			ralloc_asprintf_append(buffer, "#extension GL_EXT_shader_framebuffer_fetch : enable\n");
 			ralloc_asprintf_append(buffer, "#endif\n");
-
 			ralloc_asprintf_append(buffer, "#ifdef UE_MRT_PLS\n");
 			ralloc_asprintf_append(buffer, "#extension GL_EXT_shader_pixel_local_storage : enable\n");
-			ralloc_asprintf_append(buffer, "#endif\n");	
+			ralloc_asprintf_append(buffer, "#endif\n");
+		}	
+		
+		if (bUsesDepthbufferFetchES2 || FramebufferFetchMask || FramebufferFetchWriteMask)		
+		{
+			ralloc_asprintf_append(buffer, "\n#ifdef GL_ARM_shader_framebuffer_fetch_depth_stencil\n");
+			ralloc_asprintf_append(buffer, "#extension GL_ARM_shader_framebuffer_fetch_depth_stencil : enable\n");
+			if (!(bUsesFramebufferFetchES2 || FramebufferFetchMask || FramebufferFetchWriteMask))
+			{
+				ralloc_asprintf_append(buffer, "#elif defined(GL_EXT_shader_framebuffer_fetch)\n");
+				ralloc_asprintf_append(buffer, "#extension GL_EXT_shader_framebuffer_fetch : enable\n");
+				ralloc_asprintf_append(buffer, "#define %sinout\n", FBF_StorageQualifier);
+			}
+			ralloc_asprintf_append(buffer, "#endif\n");
 		}
 
 		// Cube map arrays are required for deferred rendering on mobile
@@ -3223,9 +3220,11 @@ class ir_gen_glsl_visitor : public ir_visitor
 		ralloc_asprintf_append(buffer, "#extension GL_EXT_texture_cube_map_array : enable\n");
 		ralloc_asprintf_append(buffer, "\n#endif\n");
 
-		if (bUsesDepthbufferFetchES2 || FramebufferFetchMask || FramebufferFetchWriteMask)
+		if (bUsesFramebufferFetchES2 || bUsesDepthbufferFetchES2)
 		{
-			ralloc_asprintf_append(buffer, "#extension GL_ARM_shader_framebuffer_fetch_depth_stencil : enable\n");
+			ralloc_asprintf_append(buffer, "\n#ifndef %s\n", FBF_StorageQualifier);
+			ralloc_asprintf_append(buffer, "#define %sout\n", FBF_StorageQualifier);
+			ralloc_asprintf_append(buffer, "#endif\n");
 		}
 
 		if (CompileTarget == HCT_FeatureLevelES3_1Ext)
@@ -3309,10 +3308,8 @@ public:
 		, needs_semicolon(false)
 		, should_print_uint_literals_as_ints(false)
 		, loop_count(0)
-		, bUsesGLTextureLODExtension(false)
 		, bUsesTextureBuffer(false)
 		, bUseImageAtomic(false)
-		, bUsesDXDY(false)
 		, bUsesDiscard(false)
 		, bUsesInstanceID(false)
 		, bNoGlobalUniforms(bInNoGlobalUniforms)
@@ -3345,32 +3342,10 @@ public:
 
 		if (bEmitPrecision && !(ShaderTarget == vertex_shader))
 		{
-			// TODO: Improve this...
-			
 			const char* DefaultPrecision = bDefaultPrecisionIsHalf ? "mediump" : "highp";
 			ralloc_asprintf_append(&default_precision_buffer, "precision %s float;\n", DefaultPrecision);
 			// always use highp for integers as shaders use them as bit storage
 			ralloc_asprintf_append(&default_precision_buffer, "precision %s int;\n", "highp");
-
-			if (bIsES) // ES workarounds
-			{
-				ralloc_asprintf_append(buffer, "\n#ifndef DONTEMITSAMPLERDEFAULTPRECISION\n"); 
-				ralloc_asprintf_append(buffer, "precision %s sampler2D;\n", DefaultPrecision);
-				ralloc_asprintf_append(buffer, "precision %s samplerCube;\n\n", DefaultPrecision);
-				ralloc_asprintf_append(buffer, "#endif\n");
-
-				// SGX540 compiler can get upset with some operations that mix highp and mediump.
-				// this results in a shader compile fail with output "compile failed."
-				// Although the actual cause of the failure hasnt been determined this code appears to prevent
-				// compile failure for cases so far seen.
-				ralloc_asprintf_append(buffer, "\n#ifdef TEXCOORDPRECISIONWORKAROUND\n");
-				ralloc_asprintf_append(buffer, "vec4 texture2DTexCoordPrecisionWorkaround(sampler2D p, vec2 tcoord)\n");
-				ralloc_asprintf_append(buffer, "{\n");
-				ralloc_asprintf_append(buffer, "	return texture2D(p, tcoord);\n");
-				ralloc_asprintf_append(buffer, "}\n");
-				ralloc_asprintf_append(buffer, "#define texture2D texture2DTexCoordPrecisionWorkaround\n");
-				ralloc_asprintf_append(buffer, "#endif\n");
-			}
 		}
 
 		// HLSLCC_DX11ClipSpace adjustment
@@ -3409,54 +3384,6 @@ bool compiler_internal_AdjustIsFrontFacing(bool isFrontFacing)
 			}
 		}
 
-		// FramebufferFetchES2 'intrinsic'
-		if (bUsesFrameBufferFetch)
-		{
-			if (state->language_version == 310)
-			{
-				// Append framebuffer fetch function. note: we want ES3.1 to throw an error if the device does not support framebuffer fetch.
-				// Fwd declare of FramebufferFetchES2 b/c we dont have access to out_Target0 yet, 
-				// #define FramebufferFetchES2() (out_Target0) fails b/c adreno preprocessor does not consume braces at call site. i.e. (out_Target0)()			
-				ralloc_asprintf_append(buffer, "\n#ifdef UE_EXT_shader_framebuffer_fetch\n");
-				ralloc_asprintf_append(buffer, "	#define %sinout\n", ES31FrameBufferFetchStorageQualifier);
-				ralloc_asprintf_append(buffer, "	vec4 FramebufferFetchES2();\n");
-				ralloc_asprintf_append(buffer, "#elif defined( GL_ARM_shader_framebuffer_fetch)\n");
-				ralloc_asprintf_append(buffer, "	#define %sout\n", ES31FrameBufferFetchStorageQualifier);
-				//ralloc_asprintf_append(buffer, "	highp vec4 gl_LastFragColorARM;\n"); // Does not work on S7 MALI
-				ralloc_asprintf_append(buffer, "	vec4 FramebufferFetchES2() { return gl_LastFragColorARM; }\n");
-				ralloc_asprintf_append(buffer, "#else\n");
-				ralloc_asprintf_append(buffer, "	#define %sout\n", ES31FrameBufferFetchStorageQualifier);
-				ralloc_asprintf_append(buffer, "	vec4 FramebufferFetchES2() { return vec4(0.0, 0.0, 0.0, 0.0); }\n");
-				ralloc_asprintf_append(buffer, "#endif\n\n");
-			}
-			else // ES3
-			{
-				ralloc_asprintf_append(buffer, "\n#ifdef UE_EXT_shader_framebuffer_fetch\n");
-				ralloc_asprintf_append(buffer, "	#if (__VERSION__ >= 300)\n");
-				ralloc_asprintf_append(buffer, "		vec4 FramebufferFetchES2() { return gl_FragColor; }\n");
-				ralloc_asprintf_append(buffer, "	#else\n");
-				ralloc_asprintf_append(buffer, "		vec4 FramebufferFetchES2() { return gl_LastFragData[0]; }\n");
-				ralloc_asprintf_append(buffer, "	#endif\n");
-				ralloc_asprintf_append(buffer, "#else\n");
-				ralloc_asprintf_append(buffer, "	#ifdef GL_ARM_shader_framebuffer_fetch\n");
-				//ralloc_asprintf_append(buffer, "		highp vec4 gl_LastFragColorARM;\n"); Does not work on S7 MALI
-				ralloc_asprintf_append(buffer, "		vec4 FramebufferFetchES2() { return gl_LastFragColorARM; }\n");
-				ralloc_asprintf_append(buffer, "	#else\n");
-				ralloc_asprintf_append(buffer, "		vec4 FramebufferFetchES2() { return vec4(0.0, 0.0, 0.0, 0.0); }\n");
-				ralloc_asprintf_append(buffer, "	#endif\n");
-				ralloc_asprintf_append(buffer, "#endif\n\n");
-			}
-		}
-
-		if (bUsesDepthbufferFetch)
-		{
-			ralloc_asprintf_append(buffer, "\n#ifdef GL_ARM_shader_framebuffer_fetch_depth_stencil\n");
-			ralloc_asprintf_append(buffer, "float DepthbufferFetchES2() { return gl_LastFragDepthARM; }\n");
-			ralloc_asprintf_append(buffer, "#else\n");
-			ralloc_asprintf_append(buffer, "float DepthbufferFetchES2() { return 0.0; }\n");
-			ralloc_asprintf_append(buffer, "#endif\n\n");
-		}
-
 		if (FramebufferFetchMask || FramebufferFetchWriteMask)
 		{
 			// Framebuffer Depth Fetch
@@ -3549,6 +3476,7 @@ bool compiler_internal_AdjustIsFrontFacing(bool isFrontFacing)
 		}
 		// End Pixel Local Storage
 
+
 		foreach_iter(exec_list_iterator, iter, *ir)
 		{
 			ir_instruction *inst = (ir_instruction *)iter.get();
@@ -3558,14 +3486,7 @@ bool compiler_internal_AdjustIsFrontFacing(bool isFrontFacing)
 
 		char* code_footer = ralloc_asprintf(mem_ctx, "");
 		buffer = &code_footer;
-		if (bUsesFrameBufferFetch && state->language_version == 310)
-		{
-			// Append ES3.1's framebuffer fetch code after the main function.
-			// It has been fwd declared earlier. Now we have access to out_Target0.
-			ralloc_asprintf_append(buffer, "\n#ifdef UE_EXT_shader_framebuffer_fetch\n");
-			ralloc_asprintf_append(buffer, "	vec4 FramebufferFetchES2() {return out_Target0;}\n");
-			ralloc_asprintf_append(buffer, "#endif\n");
-		}
+		//
 		buffer = 0;
 
 		char* decl_buffer = ralloc_asprintf(mem_ctx, "");
@@ -5902,7 +5823,7 @@ void FGlslLanguageSpec::SetupLanguageIntrinsics(_mesa_glsl_parse_state* State, e
 {
 	unsigned IntrinsicReturnType = bDefaultPrecisionIsHalf ? IR_INTRINSIC_ALL_FLOATING : IR_INTRINSIC_FLOAT;
 	make_intrinsic_genType(ir, State, FRAMEBUFFER_FETCH_ES2, ir_invalid_opcode, IntrinsicReturnType, 0, 4, 4);
-	make_intrinsic_genType(ir, State, DEPTHBUFFER_FETCH_ES2, ir_invalid_opcode, IntrinsicReturnType, 0, 1, 1);
+	make_intrinsic_genType(ir, State, DEPTHBUFFER_FETCH_ES2, ir_invalid_opcode, IR_INTRINSIC_FLOAT, 0, 1, 1);
 
 	{
 		ir_function* func = new(State)ir_function("compiler_internal_AdjustInputSemantic");
diff --git a/Engine/Source/Developer/ShaderFormatOpenGL/Private/OpenGLShaderCompiler.cpp b/Engine/Source/Developer/ShaderFormatOpenGL/Private/OpenGLShaderCompiler.cpp
--- a/Engine/Source/Developer/ShaderFormatOpenGL/Private/OpenGLShaderCompiler.cpp
+++ b/Engine/Source/Developer/ShaderFormatOpenGL/Private/OpenGLShaderCompiler.cpp
@@ -2494,7 +2494,6 @@ struct FDeviceCapabilities
 {
 	EPlatformType TargetPlatform = EPlatformType::Android;
 	bool bSupportsSeparateShaderObjects;
-	bool bRequiresUEShaderFramebufferFetchDef;
 };
 
 void FOpenGLFrontend::FillDeviceCapsOfflineCompilation(struct FDeviceCapabilities& Capabilities, const GLSLVersion ShaderVersion) const
@@ -2504,7 +2503,6 @@ void FOpenGLFrontend::FillDeviceCapsOfflineCompilation(struct FDeviceCapabilitie
 	if (ShaderVersion == GLSL_ES3_1_ANDROID)
 	{
 		Capabilities.TargetPlatform = EPlatformType::Android;
-		Capabilities.bRequiresUEShaderFramebufferFetchDef = true;
 	}
 	else
 	{
@@ -2592,11 +2590,6 @@ TSharedPtr<ANSICHAR> FOpenGLFrontend::PrepareCodeForOfflineCompilation(const GLS
 	const GLenum TypeEnum = GLFrequencyTable[Frequency];
 	// The incoming glsl may have preprocessor code that is dependent on defines introduced via the engine.
 	// This is the place to insert such engine preprocessor defines, immediately after the glsl version declaration.
-	if (Capabilities.bRequiresUEShaderFramebufferFetchDef && TypeEnum == GL_FRAGMENT_SHADER)
-	{
-		// Mali offline shader compiler does not support GL_EXT_shader_framebuffer_fetch
-		//StrOutSource.Append(TEXT("#define UE_EXT_shader_framebuffer_fetch 1\n"));
-	}
 
 	if (bEmitMobileMultiView)
 	{
diff --git a/Engine/Source/Developer/ShaderFormatOpenGL/Private/ShaderFormatOpenGL.cpp b/Engine/Source/Developer/ShaderFormatOpenGL/Private/ShaderFormatOpenGL.cpp
--- a/Engine/Source/Developer/ShaderFormatOpenGL/Private/ShaderFormatOpenGL.cpp
+++ b/Engine/Source/Developer/ShaderFormatOpenGL/Private/ShaderFormatOpenGL.cpp
@@ -17,7 +17,7 @@ class FShaderFormatGLSL : public IShaderFormat
 	enum
 	{
 		/** Version for shader format, this becomes part of the DDC key. */
-		UE_SHADER_GLSL_VER = 97,
+		UE_SHADER_GLSL_VER = 98,
 	};
 
 	void CheckFormat(FName Format) const
diff --git a/Engine/Source/Runtime/OpenGLDrv/Private/OpenGLES.cpp b/Engine/Source/Runtime/OpenGLDrv/Private/OpenGLES.cpp
--- a/Engine/Source/Runtime/OpenGLDrv/Private/OpenGLES.cpp
+++ b/Engine/Source/Runtime/OpenGLDrv/Private/OpenGLES.cpp
@@ -87,8 +87,6 @@ bool FOpenGLES::bSupportsShaderFramebufferFetch = false;
 /** GL_EXT_shader_framebuffer_fetch (MRT's) */
 bool FOpenGLES::bSupportsShaderMRTFramebufferFetch = false;
 
-/* This is to avoid a bug where device supports GL_EXT_shader_framebuffer_fetch but does not define it in GLSL */
-bool FOpenGLES::bRequiresUEShaderFramebufferFetchDef = false;
 
 /** GL_ARM_shader_framebuffer_fetch_depth_stencil */
 bool FOpenGLES::bSupportsShaderDepthStencilFetch = false;
@@ -212,7 +210,6 @@ void FOpenGLES::ProcessExtensions(const FString& ExtensionsString)
 	bSupportsShaderFramebufferFetch = ExtensionsString.Contains(TEXT("GL_EXT_shader_framebuffer_fetch")) || ExtensionsString.Contains(TEXT("GL_NV_shader_framebuffer_fetch"))
 		|| ExtensionsString.Contains(TEXT("GL_ARM_shader_framebuffer_fetch ")); // has space at the end to exclude GL_ARM_shader_framebuffer_fetch_depth_stencil match
 	bSupportsShaderMRTFramebufferFetch = ExtensionsString.Contains(TEXT("GL_EXT_shader_framebuffer_fetch")) || ExtensionsString.Contains(TEXT("GL_NV_shader_framebuffer_fetch"));
-	bRequiresUEShaderFramebufferFetchDef = ExtensionsString.Contains(TEXT("GL_EXT_shader_framebuffer_fetch"));
 	bSupportsPixelLocalStorage = ExtensionsString.Contains(TEXT("GL_EXT_shader_pixel_local_storage"));
 	bSupportsShaderDepthStencilFetch = ExtensionsString.Contains(TEXT("GL_ARM_shader_framebuffer_fetch_depth_stencil"));
 	bSupportsMultisampledRenderToTexture = ExtensionsString.Contains(TEXT("GL_EXT_multisampled_render_to_texture"));
diff --git a/Engine/Source/Runtime/OpenGLDrv/Private/OpenGLShaders.cpp b/Engine/Source/Runtime/OpenGLDrv/Private/OpenGLShaders.cpp
--- a/Engine/Source/Runtime/OpenGLDrv/Private/OpenGLShaders.cpp
+++ b/Engine/Source/Runtime/OpenGLDrv/Private/OpenGLShaders.cpp
@@ -842,11 +842,6 @@ void OPENGLDRV_API GetCurrentOpenGLShaderDeviceCapabilities(FOpenGLShaderDeviceC
 #endif
 	Capabilities.MaxRHIShaderPlatform = GMaxRHIShaderPlatform;
 	Capabilities.bSupportsSeparateShaderObjects = FOpenGL::SupportsSeparateShaderObjects();
-
-#if OPENGL_ES
-	Capabilities.bRequiresUEShaderFramebufferFetchDef = FOpenGL::RequiresUEShaderFramebufferFetchDef();
-#endif
-	
 }
 
 void OPENGLDRV_API GLSLToDeviceCompatibleGLSL(FAnsiCharArray& GlslCodeOriginal, const FString& ShaderName, GLenum TypeEnum, const FOpenGLShaderDeviceCapabilities& Capabilities, FAnsiCharArray& GlslCode)
@@ -890,12 +885,6 @@ void OPENGLDRV_API GLSLToDeviceCompatibleGLSL(FAnsiCharArray& GlslCodeOriginal,
 	// This is the place to insert such engine preprocessor defines, immediately after the glsl version declaration.
 	if (TypeEnum == GL_FRAGMENT_SHADER)
 	{
-		if (Capabilities.bRequiresUEShaderFramebufferFetchDef)
-		{
-			// Some devices (Zenfone5) support GL_EXT_shader_framebuffer_fetch but do not define GL_EXT_shader_framebuffer_fetch in GLSL compiler
-			// We can't define anything with GL_, so we use UE_EXT_shader_framebuffer_fetch to enable frame buffer fetch
-			AppendCString(GlslCode, "#define UE_EXT_shader_framebuffer_fetch 1\n");
-		}
 		if(FOpenGL::SupportsShaderMRTFramebufferFetch())
 		{
 			AppendCString(GlslCode, "#define UE_MRT_FRAMEBUFFER_FETCH 1\n");
diff --git a/Engine/Source/Runtime/OpenGLDrv/Public/OpenGLES.h b/Engine/Source/Runtime/OpenGLDrv/Public/OpenGLES.h
--- a/Engine/Source/Runtime/OpenGLDrv/Public/OpenGLES.h
+++ b/Engine/Source/Runtime/OpenGLDrv/Public/OpenGLES.h
@@ -218,7 +218,6 @@ struct FOpenGLES : public FOpenGLBase
 	static FORCEINLINE bool SupportsBufferStorage() { return bSupportsBufferStorage; }
 	
 
-	static FORCEINLINE bool RequiresUEShaderFramebufferFetchDef() { return bRequiresUEShaderFramebufferFetchDef; }
 	static FORCEINLINE bool HasBinaryProgramRetrievalFailed() { return bBinaryProgramRetrievalFailed; }
 	static FORCEINLINE bool RequiresDisabledEarlyFragmentTests() { return bRequiresDisabledEarlyFragmentTests; }
 	static FORCEINLINE bool RequiresARMShaderFramebufferFetchDepthStencilUndef() { return bRequiresARMShaderFramebufferFetchDepthStencilUndef; }
@@ -770,8 +769,6 @@ protected:
 	/** GL_EXT_shader_framebuffer_fetch (MRT's) */
 	static bool bSupportsShaderMRTFramebufferFetch;
 
-	/** workaround for GL_EXT_shader_framebuffer_fetch */
-	static bool bRequiresUEShaderFramebufferFetchDef;
 
 	/** GL_ARM_shader_framebuffer_fetch_depth_stencil */
 	static bool bSupportsShaderDepthStencilFetch;
diff --git a/Engine/Source/Runtime/OpenGLDrv/Public/OpenGLShaders.h b/Engine/Source/Runtime/OpenGLDrv/Public/OpenGLShaders.h
--- a/Engine/Source/Runtime/OpenGLDrv/Public/OpenGLShaders.h
+++ b/Engine/Source/Runtime/OpenGLDrv/Public/OpenGLShaders.h
@@ -26,7 +26,6 @@ struct FOpenGLShaderDeviceCapabilities
 	EShaderPlatform MaxRHIShaderPlatform;
 	bool bSupportsSeparateShaderObjects;
 	bool bSupportsShaderFramebufferFetch;
-	bool bRequiresUEShaderFramebufferFetchDef;
 	bool bRequiresARMShaderFramebufferFetchDepthStencilUndef;
 	GLint MaxVaryingVectors;
 	bool bRequiresDisabledEarlyFragmentTests;
