Fix SRV shader reflection in GLES backend.

#rnx

diff --git a/Engine/Source/Developer/ShaderFormatOpenGL/Private/OpenGLShaderCompiler.cpp b/Engine/Source/Developer/ShaderFormatOpenGL/Private/OpenGLShaderCompiler.cpp
--- a/Engine/Source/Developer/ShaderFormatOpenGL/Private/OpenGLShaderCompiler.cpp
+++ b/Engine/Source/Developer/ShaderFormatOpenGL/Private/OpenGLShaderCompiler.cpp
@@ -2004,7 +2004,7 @@ static bool CompileToGlslWithShaderConductor(
 		{
 		case GLSL_150_ES3_1:	// ES3.1 Emulation
 			TargetDesc.Language = CrossCompiler::EShaderConductorLanguage::Glsl;
-			TargetDesc.Version = 330;
+			TargetDesc.Version = 430;
 			break;
 		case GLSL_SWITCH_FORWARD:
 			TargetDesc.Language = CrossCompiler::EShaderConductorLanguage::Essl;
@@ -2209,14 +2209,15 @@ static bool CompileToGlslWithShaderConductor(
 				{
 					TArray<FString> UsedSamplers;
 					FString SamplerString;
-					for (FString& Sampler : Samplers)
+					for (const FString& Sampler : Samplers)
 					{
-						std::string SamplerName = TCHAR_TO_ANSI(*(Texture + Sampler));
+						std::string SamplerName = "SPIRV_Cross_Combined";
+						SamplerName  += TCHAR_TO_ANSI(*(Texture + Sampler));
 						size_t FindCombinedSampler = GlslSource.find(SamplerName.c_str());
 						if (FindCombinedSampler != std::string::npos)
 						{
 							uint32 NewIndex = TextureIndex + UsedSamplers.Num();
-							std::string NewDefine = "#define SPIRV_Cross_Combined";
+							std::string NewDefine = "#define ";
 							NewDefine += SamplerName;
 							NewDefine += " ";
 							NewDefine += FrequencyPrefix;
