Added packed UB support to DXC and fixed bugs

diff --git a/Engine/Shaders/Private/GeometryCacheVertexFactory.ush b/Engine/Shaders/Private/GeometryCacheVertexFactory.ush
--- a/Engine/Shaders/Private/GeometryCacheVertexFactory.ush
+++ b/Engine/Shaders/Private/GeometryCacheVertexFactory.ush
@@ -54,7 +54,16 @@ struct FVertexFactoryInput
     
 
 #if NUM_MATERIAL_TEXCOORDS_VERTEX
-	float2	TexCoords[NUM_MATERIAL_TEXCOORDS_VERTEX] : ATTRIBUTE5;
+	float2	TexCoords0 : ATTRIBUTE5;
+	#if NUM_MATERIAL_TEXCOORDS_VERTEX > 1
+		float2	TexCoords1 : ATTRIBUTE6;
+	#endif
+	#if NUM_MATERIAL_TEXCOORDS_VERTEX > 2
+		float2	TexCoords2 : ATTRIBUTE7;
+	#endif
+	#if NUM_MATERIAL_TEXCOORDS_VERTEX > 3
+		float2	TexCoords3 : ATTRIBUTE8;
+	#endif
 
 	#if NUM_MATERIAL_TEXCOORDS_VERTEX > 4
 		#error Too many texture coordinate sets defined on Geometry cache vertex input. Max: 4.
@@ -238,10 +247,16 @@ FMaterialVertexParameters GetMaterialVertexParameters(FVertexFactoryInput Input,
 	Result.PreSkinnedNormal = Input.TangentZ.xyz * 2.f - 1.f;
 
 #if NUM_MATERIAL_TEXCOORDS_VERTEX
-	for(int CoordinateIndex = 0; CoordinateIndex < NUM_MATERIAL_TEXCOORDS_VERTEX; CoordinateIndex++)
-	{
-		Result.TexCoords[CoordinateIndex] = Input.TexCoords[CoordinateIndex];
-	}
+	Result.TexCoords[0] = Input.TexCoords0;
+	#if NUM_MATERIAL_TEXCOORDS_VERTEX > 1
+		Result.TexCoords[1] = Input.TexCoords1;
+	#endif
+	#if NUM_MATERIAL_TEXCOORDS_VERTEX > 2
+		Result.TexCoords[2] = Input.TexCoords2;
+	#endif
+	#if NUM_MATERIAL_TEXCOORDS_VERTEX > 3
+		Result.TexCoords[3] = Input.TexCoords3;
+	#endif
 #endif
 	return Result;
 }
diff --git a/Engine/Shaders/Private/GpuSkinVertexFactory.ush b/Engine/Shaders/Private/GpuSkinVertexFactory.ush
--- a/Engine/Shaders/Private/GpuSkinVertexFactory.ush
+++ b/Engine/Shaders/Private/GpuSkinVertexFactory.ush
@@ -80,7 +80,18 @@ struct FVertexFactoryInput
 
 #if NUM_MATERIAL_TEXCOORDS_VERTEX
 	// If this changes make sure to update LocalVertexFactory.usf
-	float2	TexCoords[NUM_MATERIAL_TEXCOORDS_VERTEX] : ATTRIBUTE5;
+#if NUM_MATERIAL_TEXCOORDS_VERTEX > 0
+	float2	TexCoords0 : ATTRIBUTE5;
+#endif
+#if NUM_MATERIAL_TEXCOORDS_VERTEX > 1
+	float2	TexCoords1 : ATTRIBUTE6;
+#endif
+#if NUM_MATERIAL_TEXCOORDS_VERTEX > 2
+	float2	TexCoords2 : ATTRIBUTE7;
+#endif
+#if NUM_MATERIAL_TEXCOORDS_VERTEX > 3
+	float2	TexCoords3 : ATTRIBUTE8;
+#endif
 
 	#if NUM_MATERIAL_TEXCOORDS_VERTEX > 4
 		#error Too many texture coordinate sets defined on GPUSkin vertex input. Max: 4.
@@ -329,10 +340,20 @@ FMaterialVertexParameters GetMaterialVertexParameters(FVertexFactoryInput Input,
 	Result.PrevFrameLocalToWorld = Intermediates.PrevLocalToWorld;
 
 #if NUM_MATERIAL_TEXCOORDS_VERTEX
-	for(int CoordinateIndex = 0; CoordinateIndex < NUM_MATERIAL_TEXCOORDS_VERTEX; CoordinateIndex++)
-	{
-		Result.TexCoords[CoordinateIndex] = Input.TexCoords[CoordinateIndex];
-	}
+
+#if NUM_MATERIAL_TEXCOORDS_VERTEX > 0
+	Result.TexCoords[0] = Input.TexCoords0;
+#endif
+#if NUM_MATERIAL_TEXCOORDS_VERTEX > 1
+	Result.TexCoords[1] = Input.TexCoords1;
+#endif
+#if NUM_MATERIAL_TEXCOORDS_VERTEX > 2
+	Result.TexCoords[2] = Input.TexCoords2;
+#endif
+#if NUM_MATERIAL_TEXCOORDS_VERTEX > 3
+	Result.TexCoords[3] = Input.TexCoords3;
+#endif
+
 #endif
 
 	return Result;
diff --git a/Engine/Shaders/Private/LocalVertexFactory.ush b/Engine/Shaders/Private/LocalVertexFactory.ush
--- a/Engine/Shaders/Private/LocalVertexFactory.ush
+++ b/Engine/Shaders/Private/LocalVertexFactory.ush
@@ -105,7 +105,16 @@ struct FVertexFactoryInput
 	#if !MANUAL_VERTEX_FETCH
 		#if GPUSKIN_PASS_THROUGH
 			// These must match GPUSkinVertexFactory.usf
-			float2	TexCoords[NUM_MATERIAL_TEXCOORDS_VERTEX] : ATTRIBUTE4;
+			float2	TexCoords0 : ATTRIBUTE4;
+			#if NUM_MATERIAL_TEXCOORDS_VERTEX > 1
+				float2	TexCoords1 : ATTRIBUTE5;
+			#endif
+			#if NUM_MATERIAL_TEXCOORDS_VERTEX > 2
+				float2	TexCoords2 : ATTRIBUTE6;
+			#endif
+			#if NUM_MATERIAL_TEXCOORDS_VERTEX > 3
+				float2	TexCoords3 : ATTRIBUTE7;
+			#endif
 			#if NUM_MATERIAL_TEXCOORDS_VERTEX > 4
 				#error Too many texture coordinate sets defined on GPUSkin vertex input. Max: 4.
 			#endif
@@ -137,7 +146,7 @@ struct FVertexFactoryInput
 		#endif
 	#endif
 #elif USE_PARTICLE_SUBUVS && !MANUAL_VERTEX_FETCH
-	float2	TexCoords[1] : ATTRIBUTE4;
+	float2	TexCoords0 : ATTRIBUTE4;
 #endif
 	// Dynamic instancing related attributes with InstanceIdOffset : ATTRIBUTE13
 	VF_GPUSCENE_DECLARE_INPUT_BLOCK(13)
@@ -567,11 +576,16 @@ FMaterialVertexParameters GetMaterialVertexParameters(FVertexFactoryInput Input,
 		}
 #elif NUM_MATERIAL_TEXCOORDS_VERTEX
 		#if GPUSKIN_PASS_THROUGH
-			UNROLL
-			for (int CoordinateIndex = 0; CoordinateIndex < NUM_MATERIAL_TEXCOORDS_VERTEX; CoordinateIndex++)
-			{
-				Result.TexCoords[CoordinateIndex] = Input.TexCoords[CoordinateIndex].xy;
-			}
+			Result.TexCoords[0] = Input.TexCoords0.xy;
+			#if NUM_MATERIAL_TEXCOORDS_VERTEX > 1
+				Result.TexCoords[1] = Input.TexCoords1.xy;
+			#endif
+			#if NUM_MATERIAL_TEXCOORDS_VERTEX > 2
+				Result.TexCoords[2] = Input.TexCoords2.xy;
+			#endif
+			#if NUM_MATERIAL_TEXCOORDS_VERTEX > 3
+				Result.TexCoords[3] = Input.TexCoords3.xy;
+			#endif
 		#else
 			#if NUM_MATERIAL_TEXCOORDS_VERTEX > 0
 				Result.TexCoords[0] = Input.TexCoords0.xy;
@@ -988,7 +1002,7 @@ FVertexFactoryInterpolantsVSToPS VertexFactoryGetInterpolantsVSToPS(FVertexFacto
 	#if MANUAL_VERTEX_FETCH
 		SetUV(Interpolants, 0, LocalVF.VertexFetch_TexCoordBuffer[LocalVF.VertexFetch_Parameters[VF_NumTexcoords_Index] * (LocalVF.VertexFetch_Parameters[VF_VertexOffset] + Input.VertexId)]);
 	#else
-		SetUV(Interpolants, 0, Input.TexCoords[0]);
+		SetUV(Interpolants, 0, Input.TexCoords0);
 	#endif
 #endif
 
diff --git a/Engine/Shaders/Private/MeshParticleVertexFactory.ush b/Engine/Shaders/Private/MeshParticleVertexFactory.ush
--- a/Engine/Shaders/Private/MeshParticleVertexFactory.ush
+++ b/Engine/Shaders/Private/MeshParticleVertexFactory.ush
@@ -44,9 +44,18 @@ struct FVertexFactoryInput
 	HALF4_TYPE	VertexColor : ATTRIBUTE3;
 
 #if NUM_MATERIAL_TEXCOORDS_VERTEX
-	float2	TexCoords[NUM_MATERIAL_TEXCOORDS_VERTEX] : ATTRIBUTE4;
+	float2	TexCoords0 : ATTRIBUTE4;
+	#if NUM_MATERIAL_TEXCOORDS_VERTEX > 1
+		float2	TexCoords1 : ATTRIBUTE5;
+	#endif
+	#if NUM_MATERIAL_TEXCOORDS_VERTEX > 2
+		float2	TexCoords2 : ATTRIBUTE6;
+	#endif
+	#if NUM_MATERIAL_TEXCOORDS_VERTEX > 3
+		float2	TexCoords3 : ATTRIBUTE7;
+	#endif
 #elif USE_PARTICLE_SUBUVS
-	float2	TexCoords[1] : ATTRIBUTE4;
+	float2	TexCoords0 : ATTRIBUTE4;
 #endif
 
 #if PARTICLE_MESH_INSTANCED
@@ -411,20 +420,26 @@ FVertexFactoryIntermediates GetVertexFactoryIntermediates(FVertexFactoryInput In
 #endif
 
 #if NUM_MATERIAL_TEXCOORDS_VERTEX
-	for(int CoordinateIndex = 0; CoordinateIndex < NUM_MATERIAL_TEXCOORDS_VERTEX; CoordinateIndex++)
-	{
-		Intermediates.TexCoords[CoordinateIndex] = Input.TexCoords[CoordinateIndex].xy;
-	}
+	Intermediates.TexCoords[0] = Input.TexCoords0;
+	#if NUM_MATERIAL_TEXCOORDS_VERTEX > 1
+		Intermediates.TexCoords[1] = Input.TexCoords1;
+	#endif
+	#if NUM_MATERIAL_TEXCOORDS_VERTEX > 2
+		Intermediates.TexCoords[2] = Input.TexCoords2;
+	#endif
+	#if NUM_MATERIAL_TEXCOORDS_VERTEX > 3
+		Intermediates.TexCoords[3] = Input.TexCoords3;
+	#endif
 #endif
 
 #if USE_PARTICLE_SUBUVS
 	#if PARTICLE_MESH_INSTANCED
-		Intermediates.SubUVCoords.xy = (float2( Input.SubUVParams.x, Input.SubUVParams.y ) + Input.TexCoords[0].xy) * MeshParticleVF.SubImageSize.xy;
-		Intermediates.SubUVCoords.zw = (float2( Input.SubUVParams.z, Input.SubUVParams.w ) + Input.TexCoords[0].xy) * MeshParticleVF.SubImageSize.xy;
+		Intermediates.SubUVCoords.xy = (float2( Input.SubUVParams.x, Input.SubUVParams.y ) + Input.TexCoords0.xy) * MeshParticleVF.SubImageSize.xy;
+		Intermediates.SubUVCoords.zw = (float2( Input.SubUVParams.z, Input.SubUVParams.w ) + Input.TexCoords0.xy) * MeshParticleVF.SubImageSize.xy;
 		Intermediates.SubUVLerp = Input.SubUVLerpAndRelTime.x;
 	#else
-		Intermediates.SubUVCoords.xy = (float2( SubUVParams.x, SubUVParams.y ) + Input.TexCoords[0].xy) * MeshParticleVF.SubImageSize.xy;
-		Intermediates.SubUVCoords.zw = (float2( SubUVParams.z, SubUVParams.w ) + Input.TexCoords[0].xy) * MeshParticleVF.SubImageSize.xy;
+		Intermediates.SubUVCoords.xy = (float2( SubUVParams.x, SubUVParams.y ) + Input.TexCoords0.xy) * MeshParticleVF.SubImageSize.xy;
+		Intermediates.SubUVCoords.zw = (float2( SubUVParams.z, SubUVParams.w ) + Input.TexCoords0.xy) * MeshParticleVF.SubImageSize.xy;
 		Intermediates.SubUVLerp = SubUVLerp.x;
 	#endif
 #endif 
diff --git a/Engine/Shaders/Public/Platform.ush b/Engine/Shaders/Public/Platform.ush
--- a/Engine/Shaders/Public/Platform.ush
+++ b/Engine/Shaders/Public/Platform.ush
@@ -29,6 +29,9 @@
 #elif COMPILER_VULKAN
 	// Helps with iteration when changing Vulkan shader code generation backend.
 	#include "Platform/Vulkan/VulkanCommon.ush"
+#elif COMPILER_GLSL || COMPILER_GLSL_ES3_1
+	// Helps with iteration when changing Vulkan shader code generation backend.
+	#include "Platform/GL/GLCommon.ush"
 #elif SM6_PROFILE || SM5_PROFILE
 	#include "Platform/D3D/D3DCommon.ush"
 #endif
diff --git a/Engine/Shaders/Public/Platform/GL/GLCommon.ush b/Engine/Shaders/Public/Platform/GL/GLCommon.ush
new file mode 100644
--- /dev/null
+++ b/Engine/Shaders/Public/Platform/GL/GLCommon.ush
@@ -0,0 +1,10 @@
+// Copyright Epic Games, Inc. All Rights Reserved.
+
+/*=============================================================================
+	GLCommon.ush: Common GL shader code
+=============================================================================*/
+
+// Update this GUID to improve shader recompilation for GL only shaders
+// GUID = 58EB247E-D403-451E-B31F-4B24D38898B8
+
+#pragma once
diff --git a/Engine/Source/Developer/ShaderCompilerCommon/Private/HlslccHeaderWriter.cpp b/Engine/Source/Developer/ShaderCompilerCommon/Private/HlslccHeaderWriter.cpp
--- a/Engine/Source/Developer/ShaderCompilerCommon/Private/HlslccHeaderWriter.cpp
+++ b/Engine/Source/Developer/ShaderCompilerCommon/Private/HlslccHeaderWriter.cpp
@@ -388,7 +388,7 @@ namespace CrossCompiler
 		MetaDataPrintf(Strings.UniformBlocks, TEXT("%s(%u)"), ResourceName, BindingIndex);
 	}
 
-	static EPackedTypeName EncodePackedGlobalType(const SpvReflectTypeDescription& TypeDescription)
+	EPackedTypeName FHlslccHeaderWriter::EncodePackedGlobalType(const SpvReflectTypeDescription& TypeDescription)
 	{
 		const SpvReflectTypeFlags ScalarTypeFlagsBitmask =
 		(
@@ -429,16 +429,39 @@ namespace CrossCompiler
 
 	void FHlslccHeaderWriter::WritePackedUB(uint32 BindingIndex)
 	{
-		checkf(Strings.PackedUB.IsEmpty(), TEXT("cannot define @PackedUB attribute more than once"));
-		MetaDataPrintf(Strings.PackedUB, TEXT("Globals(%u): "), BindingIndex);
+		WritePackedUB(TEXT("Globals"), BindingIndex);
 	}
 
 	void FHlslccHeaderWriter::WritePackedUBField(const TCHAR* ResourceName, uint32 ByteOffset, uint32 ByteSize)
 	{
-		checkf(!Strings.PackedUB.IsEmpty(), TEXT("cannot append field without @PackedUB attribute in shader meta data"));
+		WritePackedUBField(TEXT("Globals"), ResourceName, ByteOffset, ByteSize);
+	}
+
+	void FHlslccHeaderWriter::WritePackedUB(const FString& UBName, uint32 BindingIndex)
+	{
+		checkf(Strings.PackedUBs.Find(UBName) == nullptr, TEXT("attempting to add a UB that has already been added"));
+		MetaDataPrintf(Strings.PackedUBs.Add(UBName), TEXT("%s(%u): "), *UBName, BindingIndex);
+	}
+
+	void FHlslccHeaderWriter::WritePackedUBField(const FString& UBName, const TCHAR* ResourceName, uint32 ByteOffset, uint32 ByteSize)
+	{
+		checkf(Strings.PackedUBs.Find(UBName), TEXT("cannot append field without @PackedUB attribute in shader meta data, %s"), *UBName);
 		checkf(ByteOffset % 4 == 0, TEXT("field offset in @PackedUB shader meta data must be a multiple of 4, but got %u"), ByteOffset);
 		checkf(ByteSize % 4 == 0, TEXT("field size in @PackedUB shader meta data must be a multiple of 4, but got %u"), ByteSize);
-		MetaDataPrintf(Strings.PackedUBFields, TEXT("%s(%u,%u)"), ResourceName, ByteOffset / 4, ByteSize / 4);
+		MetaDataPrintf(Strings.PackedUBFields.FindOrAdd(UBName), TEXT("%s(%u,%u)"), ResourceName, ByteOffset / 4, ByteSize / 4);
+	}
+
+	void FHlslccHeaderWriter::WritePackedUBCopy(uint32 SourceCB, uint32 SourceOffset, uint32 DestCBIndex, uint32 DestCBPrecision, uint32 DestOffset, uint32 Size, bool bGroupFlattenedUBs)
+	{
+		if (bGroupFlattenedUBs)
+		{
+			MetaDataPrintf(Strings.PackedUBCopies, TEXT("%u:%u-%u:%c:%u:%u"), SourceCB, SourceOffset, DestCBIndex, DestCBPrecision, DestOffset, Size);
+		}
+		else
+		{
+			check(DestCBIndex == 0);
+			MetaDataPrintf(Strings.PackedUBCopies, TEXT("%u:%u-%c:%u:%u"), SourceCB, SourceOffset, DestCBPrecision, DestOffset, Size);
+		}
 	}
 
 	void FHlslccHeaderWriter::WritePackedUBGlobalCopy(uint32 SourceCB, uint32 SourceOffset, uint32 DestCBIndex, uint32 DestCBPrecision, uint32 DestOffset, uint32 Size, bool bGroupFlattenedUBs)
@@ -546,8 +569,15 @@ namespace CrossCompiler
 		PrintAttributes(TEXT("Inputs"), Strings.InputAttributes);
 		PrintAttributes(TEXT("Outputs"), Strings.OutputAttributes);
 		PrintAttributes(TEXT("UniformBlocks"), Strings.UniformBlocks);
-		PrintAttributes(TEXT("PackedUB"), Strings.PackedUB + Strings.PackedUBFields);
+
 		PrintAttributes(TEXT("PackedGlobals"), Strings.PackedGlobals);
+
+		for (TPair<FString, FString> UBName : Strings.PackedUBs)
+		{
+			PrintAttributes(TEXT("PackedUB"), UBName.Value + Strings.PackedUBFields[UBName.Key]);
+		}
+
+		PrintAttributes(TEXT("PackedUBCopies"), Strings.PackedUBCopies);
 		PrintAttributes(TEXT("PackedUBGlobalCopies"), Strings.PackedUBGlobalCopies);
 		PrintAttributes(TEXT("Samplers"), Strings.SRVs); // Was called "Samplers" in HLSLcc but serves as SRVs
 		PrintAttributes(TEXT("UAVs"), Strings.UAVs);
diff --git a/Engine/Source/Developer/ShaderCompilerCommon/Public/HlslccHeaderWriter.h b/Engine/Source/Developer/ShaderCompilerCommon/Public/HlslccHeaderWriter.h
--- a/Engine/Source/Developer/ShaderCompilerCommon/Public/HlslccHeaderWriter.h
+++ b/Engine/Source/Developer/ShaderCompilerCommon/Public/HlslccHeaderWriter.h
@@ -9,6 +9,7 @@
 // Forward declaration from <spirv_reflect.h>
 struct SpvReflectInterfaceVariable;
 struct SpvReflectBlockVariable;
+struct SpvReflectTypeDescription;
 
 // Cross compiler support/common functionality
 namespace CrossCompiler
@@ -29,6 +30,9 @@ namespace CrossCompiler
 		void WritePackedGlobal(const SpvReflectBlockVariable& Variable);
 		void WritePackedUB(uint32 BindingIndex);
 		void WritePackedUBField(const TCHAR* ResourceName, uint32 ByteOffset, uint32 ByteSize);
+		void WritePackedUB(const FString& UBName, uint32 BindingIndex);
+		void WritePackedUBField(const FString& UBName, const TCHAR* ResourceName, uint32 ByteOffset, uint32 ByteSize);
+		void WritePackedUBCopy(uint32 SourceCB, uint32 SourceOffset, uint32 DestCBIndex, uint32 DestCBPrecision, uint32 DestOffset, uint32 Size, bool bGroupFlattenedUBs = false);
 		void WritePackedUBGlobalCopy(uint32 SourceCB, uint32 SourceOffset, uint32 DestCBIndex, uint32 DestCBPrecision, uint32 DestOffset, uint32 Size, bool bGroupFlattenedUBs = false);
 		void WriteSRV(const TCHAR* ResourceName, uint32 BindingIndex, uint32 Count = 1);
 		void WriteSRV(const TCHAR* ResourceName, uint32 BindingIndex, uint32 Count, const TArray<FString>& AssociatedResourceNames);
@@ -42,6 +46,7 @@ namespace CrossCompiler
 
 		/** Returns the finalized meta data. */
 		FString ToString() const;
+		static EPackedTypeName EncodePackedGlobalType(const SpvReflectTypeDescription& TypeDescription);
 
 	private:
 		void WriteIOAttribute(FString& OutMetaData, const TCHAR* AttributeName, const TCHAR* TypeSpecifier, int32 Location, bool bLocationPrefix, bool bLocationSuffix);
@@ -56,8 +61,9 @@ namespace CrossCompiler
 			FString OutputAttributes;
 			FString UniformBlocks;
 			FString PackedGlobals;
-			FString PackedUB;
-			FString PackedUBFields;
+			TMap<FString, FString> PackedUBs;
+			TMap<FString, FString> PackedUBFields;
+			FString PackedUBCopies;
 			FString PackedUBGlobalCopies;
 			FString SRVs; // Shader resource views (SRV) and samplers
 			FString UAVs; // Unordered access views (UAV)
diff --git a/Engine/Source/Developer/ShaderFormatOpenGL/Private/OpenGLShaderCompiler.cpp b/Engine/Source/Developer/ShaderFormatOpenGL/Private/OpenGLShaderCompiler.cpp
--- a/Engine/Source/Developer/ShaderFormatOpenGL/Private/OpenGLShaderCompiler.cpp
+++ b/Engine/Source/Developer/ShaderFormatOpenGL/Private/OpenGLShaderCompiler.cpp
@@ -698,8 +698,7 @@ void FOpenGLFrontend::BuildShaderOutput(
 	for (auto& UniformBlock : CCHeader.UniformBlocks)
 	{
 		uint16 UBIndex = UniformBlock.Index;
-		check(UBIndex == Header.Bindings.NumUniformBuffers);
-		UsedUniformBufferSlots[UBIndex] = true;
+		
 		if (OutputTrueParameterNames())
 		{
 			// make the final name this will be in the shader
@@ -709,17 +708,37 @@ void FOpenGLFrontend::BuildShaderOutput(
 		}
 		else
 		{
-			ParameterMap.AddParameterAllocation(*UniformBlock.Name, Header.Bindings.NumUniformBuffers, 0, 0, EShaderParameterType::UniformBuffer);
+			ParameterMap.AddParameterAllocation(*UniformBlock.Name, UBIndex, 0, 0, EShaderParameterType::UniformBuffer);
 		}
 		Header.Bindings.NumUniformBuffers++;
 	}
 
 	const uint16 BytesPerComponent = 4;
 
+
+	FString GlobalIgnoreStrings[] =
+	{
+		TEXT("gl_LastFragDepthARM"),
+		TEXT("ARM_shader_framebuffer_fetch_depth_stencil"),
+	};
+
 	// Packed global uniforms
 	TMap<ANSICHAR, uint16> PackedGlobalArraySize;
 	for (auto& PackedGlobal : CCHeader.PackedGlobals)
 	{
+		bool bIgnore = false;
+		for (uint32_t i = 0; i < UE_ARRAY_COUNT(GlobalIgnoreStrings); ++i)
+		{
+			if (PackedGlobal.Name.StartsWith(GlobalIgnoreStrings[i]))
+			{
+				bIgnore = true;
+				break;
+			}
+		}
+
+		if (bIgnore)
+			continue;
+
 		ParameterMap.AddParameterAllocation(
 			*PackedGlobal.Name,
 			PackedGlobal.PackedType,
@@ -738,7 +757,6 @@ void FOpenGLFrontend::BuildShaderOutput(
 	{
 		checkf(OutputTrueParameterNames() == false, TEXT("Unexpected Packed UBs used with a shader format that needs true parameter names - If this is hit, we need to figure out how to handle them"));
 
-		check(PackedUB.Attribute.Index == Header.Bindings.NumUniformBuffers);
 		UsedUniformBufferSlots[PackedUB.Attribute.Index] = true;
 		if (OutputTrueParameterNames())
 		{
@@ -747,7 +765,7 @@ void FOpenGLFrontend::BuildShaderOutput(
 		}
 		else
 		{
-			ParameterMap.AddParameterAllocation(*PackedUB.Attribute.Name, Header.Bindings.NumUniformBuffers, 0, 0, EShaderParameterType::UniformBuffer);
+			ParameterMap.AddParameterAllocation(*PackedUB.Attribute.Name, PackedUB.Attribute.Index, 0, 0, EShaderParameterType::UniformBuffer);
 		}
 		Header.Bindings.NumUniformBuffers++;
 
@@ -1221,6 +1239,228 @@ static const TCHAR* GetFrequencyFileExt(EHlslShaderFrequency Frequency)
 	}
 }
 
+struct PackedUBMemberInfo
+{
+	std::string Name;
+	std::string TypeQualifier;
+	uint32_t SrcOffset;
+	uint32_t DestOffset;
+	uint32_t SrcSizeInFloats;
+	uint32_t DestSizeInFloats;
+};
+
+void WritePackedUBHeader(CrossCompiler::FHlslccHeaderWriter& CCHeaderWriter, const TMap<uint32_t, TArray<PackedUBMemberInfo>>& UBMemberInfo, const TMap<uint32, std::string>& UBNames)
+{
+	bool bFirst = true;
+
+	bool bNeedsHeader = true;
+
+	for (const auto Pair : UBNames)
+	{
+		FString Name(Pair.Value.c_str());
+		CCHeaderWriter.WritePackedUB(Name, Pair.Key);
+	}
+
+	for (const auto Pair : UBMemberInfo)
+	{
+		FString UBName(UBNames[Pair.Key].c_str());
+		
+		for (const PackedUBMemberInfo& MemberInfo : UBMemberInfo[Pair.Key])
+		{
+			CCHeaderWriter.WritePackedUBField(UBName, UTF8_TO_TCHAR(MemberInfo.Name.c_str()), MemberInfo.SrcOffset, MemberInfo.DestSizeInFloats * sizeof(float));
+			CCHeaderWriter.WritePackedUBCopy(Pair.Key, MemberInfo.SrcOffset / sizeof(float), Pair.Key, TEXT('h'), MemberInfo.DestOffset / sizeof(float), MemberInfo.DestSizeInFloats, true);
+		}
+	}
+}
+
+void GetSpvVarQualifier(const SpvReflectBlockVariable& Member, FString & Out)
+{
+	auto const type = *Member.type_description;
+	const uint32 MbrSize = Member.size / sizeof(float);
+
+	FString TypeQualifier;
+
+	uint32_t masked_type = type.type_flags & 0xF;
+
+	switch (masked_type)
+	{
+	default: checkf(false, TEXT("unsupported component type %d"), masked_type); break;
+	case SPV_REFLECT_TYPE_FLAG_BOOL:
+	case SPV_REFLECT_TYPE_FLAG_INT: 
+		Out = (type.traits.numeric.scalar.signedness ? TEXT("i") : TEXT("u"));
+		break;
+	case SPV_REFLECT_TYPE_FLAG_FLOAT: 
+		Out = TEXT("h");
+		break;
+	}
+}
+
+// Adds a member to ne included in the PackedUB structures and generates the #define for readability in glsl
+void AddMemberToPackedUB(const std::string & FrequencyPrefix,
+							const SpvReflectBlockVariable& Member,
+							const std::string& UBName,
+							int32_t Index,
+							TMap<FString, uint32>& Offsets,
+							TArray<PackedUBMemberInfo>& MemberInfos,
+							TArray<std::string>& Remap,
+							TArray<std::string>& Arrays)
+{
+	std::string ArrayName;
+	const uint32 MbrSize = Member.size / sizeof(float);
+
+	FString TypeQualifier;
+
+	GetSpvVarQualifier(Member, TypeQualifier);
+
+	FString MemberName = UTF8_TO_TCHAR(Member.name);
+
+	uint32& Offset = Offsets.FindOrAdd(TypeQualifier);
+	
+	auto const type = *Member.type_description;
+
+	bool const bArray = type.traits.array.dims_count > 0;
+	bool const bGlobals = Index == -1;
+
+	std::string Name = "#define ";
+
+	if (bGlobals)
+	{
+		Name += "_Globals_";
+	}
+
+	std::string OffsetString = std::to_string(Offset);
+	Name += Member.name;
+
+	PackedUBMemberInfo& MemberInfo = MemberInfos.AddDefaulted_GetRef();
+	MemberInfo.Name = Member.name;
+	MemberInfo.TypeQualifier = TCHAR_TO_UTF8(*TypeQualifier);
+	MemberInfo.SrcOffset = Member.offset;
+	MemberInfo.DestOffset = Offset * 4 * sizeof(float);
+	MemberInfo.SrcSizeInFloats = Member.size / sizeof(float);
+	MemberInfo.DestSizeInFloats = Member.size / sizeof(float);	
+
+	if (bArray)
+	{
+		if (bGlobals)
+		{
+			ArrayName = UBName + Member.name;
+		}
+		else
+		{
+			ArrayName = Member.name;
+		}
+
+		Name += "(Offset)";
+		if (type.op == SpvOpTypeMatrix)
+		{
+			OffsetString += " + (int(Offset) * 4)";
+		}
+		else
+		{
+			OffsetString += " + int(Offset)";
+		}
+	}
+
+	Name += " (";
+
+	std::string UniformPrefix = FrequencyPrefix;
+	if (!bGlobals)
+	{
+		UniformPrefix += std::string("c") + std::to_string(Index);
+	}
+	else
+	{
+		UniformPrefix += "u";
+	}
+
+	if (type.op == SpvOpTypeMatrix)
+	{
+		if ((type.traits.numeric.matrix.column_count == 4 && type.traits.numeric.matrix.row_count == 4))
+		{
+			Name += "mat4(";
+		}
+		else
+		{
+			char buff[10];
+			sprintf_s(buff, 10, "mat%dx%d(", type.traits.numeric.matrix.column_count, type.traits.numeric.matrix.row_count);
+			Name += buff;
+		}
+
+		for (uint32_t i = 0; i < type.traits.numeric.matrix.row_count; ++i)
+		{
+			if (i > 0)
+			{
+				Name += ",";
+			}
+
+			Name += UniformPrefix;
+			Name += "_";
+			Name += TCHAR_TO_UTF8(*TypeQualifier);
+			
+			char buff[10];
+			sprintf_s(buff, 10, "[%s + %d]", OffsetString.c_str(), i);
+
+			Name += buff;
+
+			switch (type.traits.numeric.matrix.column_count)
+			{
+			case 0:
+			case 1:
+				Name += ".x";
+				break;
+			case 2:
+				Name += ".xy";
+				break;
+			case 3:
+				Name += ".xyz";
+				break;
+			case 4:
+			default:
+				Name += ".xyzw";
+				break;
+			}
+		}
+
+		Name += ")";
+	}
+	else
+	{
+		Name += UniformPrefix;
+		Name += "_";
+		Name += TCHAR_TO_UTF8(*TypeQualifier);
+		Name += "[";
+		Name += OffsetString;
+		Name += "].";
+		switch (type.traits.numeric.vector.component_count)
+		{
+		case 0:
+		case 1:
+			Name += "x";
+			break;
+		case 2:
+			Name += "xy";
+			break;
+		case 3:
+			Name += "xyz";
+			break;
+		case 4:
+		default:
+			Name += "xyzw";
+			break;
+		}
+	}
+	Name += ")\n";
+
+	if (bArray)
+	{
+		Arrays.Add(ArrayName);
+	}
+
+	Remap.Add(Name);
+
+	Offset += Align(MbrSize, 4) / 4;
+}
+
 static bool CompileToGlslWithShaderConductor(
 	const FShaderCompilerInput&	Input,
 	FShaderCompilerOutput&		Output,
@@ -1243,6 +1483,8 @@ static bool CompileToGlslWithShaderConductor(
 	std::string SourceData(TCHAR_TO_UTF8(*PreprocessedShader));
 	std::string FileName(TCHAR_TO_UTF8(*Input.VirtualSourceFilePath));
 	std::string EntryPointName(TCHAR_TO_UTF8(*Input.EntryPointName));
+	
+	bool bEmulatedUBs = Input.Environment.CompilerFlags.Contains(CFLAG_UseEmulatedUB);
 
 	// HLSL framebuffer declarations. Used to modify HLSL input source.
 	const ANSICHAR* HlslFrameBufferDeclarations =
@@ -1264,15 +1506,24 @@ static bool CompileToGlslWithShaderConductor(
 		"}\n"
 		"float DepthbufferFetchES2()\n"
 		"{\n"
-		"  return (!ARM_shader_framebuffer_fetch_depth_stencil ? 0.0 : gl_LastFragDepthARM);\n"
+		"  return (ARM_shader_framebuffer_fetch_depth_stencil == 0 ? 0.0 : gl_LastFragDepthARM);\n"
 		"}\n"
 		;
 
 	SourceData = HlslFrameBufferDeclarations + SourceData;
 
+	const ANSICHAR* GlslFrameBufferExtensions =
+		"\n\n#ifdef GL_ARM_shader_framebuffer_fetch_depth_stencil\n"
+		"\t#extension GL_ARM_shader_framebuffer_fetch_depth_stencil : enable\n"
+		"#elif defined(GL_EXT_shader_framebuffer_fetch)\n"
+		"\t#extension GL_EXT_shader_framebuffer_fetch : enable\n"
+		"\t#define FBF_STORAGE_QUALIFIER inout\n"
+		"#endif\n"
+		"#extension GL_EXT_texture_buffer : enable\n";
+
 	// GLSL framebuffer macro definitions. Used to patch GLSL output source.
 	const ANSICHAR* GlslFrameBufferDefines =
-		"#ifdef UE_EXT_shader_framebuffer_fetch\n"
+		"\n\n#ifdef UE_EXT_shader_framebuffer_fetch\n"
 		"#define _Globals_ARM_shader_framebuffer_fetch 0\n"
 		"#define FRAME_BUFFERFETCH_STORAGE_QUALIFIER inout\n"
 		"#define _Globals_gl_FragColor out_var_SV_Target0\n"
@@ -1289,9 +1540,9 @@ static bool CompileToGlslWithShaderConductor(
 		"#define _Globals_gl_LastFragColorARM vec4(0.0, 0.0, 0.0, 0.0)\n"
 		"#endif\n"
 		"#ifdef GL_ARM_shader_framebuffer_fetch_depth_stencil\n"
-		"#define _Globals_ARM_shader_framebuffer_fetch_depth_stencil 1\n"
+		"#define _Globals_ARM_shader_framebuffer_fetch_depth_stencil 1u\n"
 		"#else\n"
-		"#define _Globals_ARM_shader_framebuffer_fetch_depth_stencil 0\n"
+		"#define _Globals_ARM_shader_framebuffer_fetch_depth_stencil 0u\n"
 		"#endif\n"
 		;
 	
@@ -1453,6 +1704,19 @@ static bool CompileToGlslWithShaderConductor(
 			check(SPVRResult == SPV_REFLECT_RESULT_SUCCESS);
 		}
 
+		TMap<FString, uint32> GlobalOffsets;
+		TArray<std::string> GlobalRemap;
+		TArray<std::string> GlobalArrays;
+		TArray<PackedUBMemberInfo> GlobalMemberInfos;
+
+		TMap<uint32, TMap<FString, uint32>> PackedUBOffsets;
+		TMap<uint32, TArray<std::string>> PackedUBRemap;
+		TMap<uint32, TArray<std::string>> PackedUBArrays;
+		TMap<uint32, TArray<PackedUBMemberInfo>> PackedUBMemberInfos;
+		TMap<uint32, std::string> PackedUBNames;
+
+		uint32_t PackedUBIndex = 0;
+
 		for (auto const& Binding : ReflectionBindings.UniformBuffers)
 		{
 			check(UBOIndices);
@@ -1462,27 +1726,83 @@ static bool CompileToGlslWithShaderConductor(
 			// Global uniform buffer - handled specially as we care about the internal layout
 			if (strstr(Binding->name, "$Globals"))
 			{
+				TMap<FString, uint32_t> GlobalOffsetSizes;
 				for (uint32 i = 0; i < Binding->block.member_count; i++)
 				{
-					CCHeaderWriter.WritePackedGlobal(Binding->block.members[i]);
+					SpvReflectBlockVariable& member = Binding->block.members[i];
+
+					FString TypeQualifier;
+					GetSpvVarQualifier(member, TypeQualifier);
+
+					uint32_t& GlobalOffsetSize = GlobalOffsetSizes.FindOrAdd(TypeQualifier);
+
+					if (strstr(member.name, "gl_") || !strcmp(member.name, "ARM_shader_framebuffer_fetch") || !strcmp(member.name, "ARM_shader_framebuffer_fetch_depth_stencil"))
+					{
+						continue;
+					}
+
+					CCHeaderWriter.WritePackedGlobal(ANSI_TO_TCHAR(member.name), CrossCompiler::FHlslccHeaderWriter::EncodePackedGlobalType(*(member.type_description)), GlobalOffsetSize, member.size);
+					GlobalOffsetSize += Align(member.size, 16);
+
+					AddMemberToPackedUB(FrequencyPrefix,
+						member,
+						"_Globals_",
+						-1,
+						GlobalOffsets,
+						GlobalMemberInfos,
+						GlobalRemap,
+						GlobalArrays);
 				}
 			}
 			else
 			{
-				std::string OldName = "type_";
-				OldName += Binding->name;
-				std::string NewName = FrequencyPrefix;
-				NewName += "b";
-				NewName += std::to_string(Index);
-				UniformVarNames[OldName] = NewName;
-				// Regular uniform buffer - we only care about the binding index
-				CCHeaderWriter.WriteUniformBlock(UTF8_TO_TCHAR(Binding->name), Index);
+				if (bEmulatedUBs)
+				{
+					//CCHeaderWriter.WritePackedUB(Binding->binding);
+					
+					PackedUBOffsets.Add(PackedUBIndex);
+					PackedUBRemap.Add(PackedUBIndex);
+					PackedUBArrays.Add(PackedUBIndex);
+					PackedUBMemberInfos.Add(PackedUBIndex);
+					PackedUBNames.Add(PackedUBIndex, std::string(Binding->name));
+
+					for (uint32 i = 0; i < Binding->block.member_count; i++)
+					{
+						SpvReflectBlockVariable& member = Binding->block.members[i];
+						//CCHeaderWriter.WritePackedUBField(UTF8_TO_TCHAR(member.name), member.absolute_offset, member.size);
+
+						std::string UBName = "_" + std::string(Binding->name) + "_";
+						
+						AddMemberToPackedUB(FrequencyPrefix,
+							member,
+							UBName,
+							PackedUBIndex,
+							PackedUBOffsets[PackedUBIndex],
+							PackedUBMemberInfos[PackedUBIndex],
+							PackedUBRemap[PackedUBIndex],
+							PackedUBArrays[PackedUBIndex]);
+					}
+
+					PackedUBIndex++;
+				}
+				else
+				{
+					std::string OldName = Binding->name;
+					std::string NewName = FrequencyPrefix;
+					NewName += "b";
+					NewName += std::to_string(Index);
+					UniformVarNames[OldName] = NewName;
+					// Regular uniform buffer - we only care about the binding index
+					CCHeaderWriter.WriteUniformBlock(UTF8_TO_TCHAR(Binding->name), Index);
+				}
 			}
 
 			SPVRResult = Reflection.ChangeDescriptorBindingNumbers(Binding, Index, GlobalSetId);
 			check(SPVRResult == SPV_REFLECT_RESULT_SUCCESS);
 		}
 
+		WritePackedUBHeader(CCHeaderWriter, PackedUBMemberInfos, PackedUBNames);
+
 		for (auto const& Binding : ReflectionBindings.TextureSRVs)
 		{
 			check(TextureIndices);
@@ -1507,10 +1827,6 @@ static bool CompileToGlslWithShaderConductor(
 			check(SPVRResult == SPV_REFLECT_RESULT_SUCCESS);
 		}
 
-		TArray<std::string> GlobalRemap;
-		TArray<std::string> GlobalArrays;
-		TMap<FString, uint32> GlobalOffsets;
-
 		{
 			uint32 Count = 0;
 			SPVRResult = Reflection.EnumeratePushConstantBlocks(&Count, nullptr);
@@ -1525,116 +1841,33 @@ static bool CompileToGlslWithShaderConductor(
 					// Global uniform buffer - handled specially as we care about the internal layout
 					if (strstr(Var->name, "$Globals"))
 					{
+						TMap<FString, uint32_t> GlobalOffsetSizes;
+
 						for (uint32 i = 0; i < Var->member_count; i++)
 						{
 							SpvReflectBlockVariable& member = Var->members[i];
-							
-							if(!strcmp(member.name, "gl_FragColor") || !strcmp(member.name, "gl_LastFragColorARM") || !strcmp(member.name, "gl_LastFragDepthARM") || !strcmp(member.name, "ARM_shader_framebuffer_fetch") || !strcmp(member.name, "ARM_shader_framebuffer_fetch_depth_stencil"))
-							{
-								continue;
-							}
-							auto const type = *member.type_description;
-
-							const uint32 MbrSize = member.size / sizeof(float);
-
+						
 							FString TypeQualifier;
+							GetSpvVarQualifier(member, TypeQualifier);
 
-							uint32_t masked_type = type.type_flags & 0xF;
+							uint32_t& GlobalOffsetSize = GlobalOffsetSizes.FindOrAdd(TypeQualifier);
 
-							switch (masked_type) {
-							default: checkf(false, TEXT("unsupported component type %d"), masked_type); break;
-							case SPV_REFLECT_TYPE_FLAG_BOOL:
-							case SPV_REFLECT_TYPE_FLAG_INT: TypeQualifier = (type.traits.numeric.scalar.signedness ? TEXT("i") : TEXT("u")); break;
-							case SPV_REFLECT_TYPE_FLAG_FLOAT: TypeQualifier = (TEXT("h")); break;
-							}
-
-							FString MemberName = UTF8_TO_TCHAR(member.name);
-
-							uint32& Offset = GlobalOffsets.FindOrAdd(TypeQualifier);
-
-							CCHeaderWriter.WritePackedGlobal(member);
-
-							bool const bArray = type.traits.array.dims_count > 0;
-
-							std::string Name = "#define _Globals_";
-							std::string OffsetString = std::to_string(Offset);
-							Name += member.name;
-							if (bArray)
-							{
-								std::string ArrayName = "_Globals_";
-								ArrayName += member.name;
-								GlobalArrays.Add(ArrayName);
-								Name += "(Offset)";
-								if (type.op == SpvOpTypeMatrix)
-									OffsetString += " + (Offset * 4)";
-								else
-									OffsetString += " + Offset";
-							}
-							Name += " (";
-							if (type.op == SpvOpTypeMatrix)
+							if (!strcmp(member.name, "gl_FragColor") || !strcmp(member.name, "gl_LastFragColorARM") || !strcmp(member.name, "gl_LastFragDepthARM") || !strcmp(member.name, "ARM_shader_framebuffer_fetch") || !strcmp(member.name, "ARM_shader_framebuffer_fetch_depth_stencil"))
 							{
-								check(type.traits.numeric.matrix.column_count == 4 && type.traits.numeric.matrix.row_count == 4);
-								Name += "mat4(";
-								Name += FrequencyPrefix;
-								Name += "u_";
-								Name += TCHAR_TO_UTF8(*TypeQualifier);
-								Name += "[";
-								Name += OffsetString;
-								Name += "],";
-
-								Name += FrequencyPrefix;
-								Name += "u_";
-								Name += TCHAR_TO_UTF8(*TypeQualifier);
-								Name += "[";
-								Name += OffsetString;
-								Name += "+ 1],";
-
-								Name += FrequencyPrefix;
-								Name += "u_";
-								Name += TCHAR_TO_UTF8(*TypeQualifier);
-								Name += "[";
-								Name += OffsetString;
-								Name += " + 2],";
-
-								Name += FrequencyPrefix;
-								Name += "u_";
-								Name += TCHAR_TO_UTF8(*TypeQualifier);
-								Name += "[";
-								Name += OffsetString;
-								Name += " + 3]";
-								Name += ")";
-							}
-							else
-							{
-								Name += FrequencyPrefix;
-								Name += "u_";
-								Name += TCHAR_TO_UTF8(*TypeQualifier);
-								Name += "[";
-								Name += OffsetString;
-								Name += "].";
-								switch (type.traits.numeric.vector.component_count)
-								{
-								case 0:
-								case 1:
-									Name += "x";
-									break;
-								case 2:
-									Name += "xy";
-									break;
-								case 3:
-									Name += "xyz";
-									break;
-								case 4:
-								default:
-									Name += "xyzw";
-									break;
-								}
+								continue;
 							}
-							Name += ")\n";
-
-							GlobalRemap.Add(Name);
 
-							Offset += Align(MbrSize, 4) / 4;
+							CCHeaderWriter.WritePackedGlobal(ANSI_TO_TCHAR(member.name), CrossCompiler::FHlslccHeaderWriter::EncodePackedGlobalType(*(member.type_description)), GlobalOffsetSize, member.size);
+							GlobalOffsetSize += Align(member.size, 16);
+
+							AddMemberToPackedUB(FrequencyPrefix,
+								member,
+								"_Globals_",
+								-1,
+								GlobalOffsets,
+								GlobalMemberInfos,
+								GlobalRemap,
+								GlobalArrays);
 						}
 					}
 				}
@@ -1868,6 +2101,8 @@ static bool CompileToGlslWithShaderConductor(
 			break;
 		case GLSL_ES3_1_ANDROID:
 		default:
+			TargetDesc.CompileFlags.SetDefine(TEXT("force_flattened_io_blocks"), 1);
+			TargetDesc.CompileFlags.SetDefine(TEXT("emit_uniform_buffer_as_plain_uniforms"), 1);
 			TargetDesc.Language = CrossCompiler::EShaderConductorLanguage::Essl;
 			TargetDesc.Version = 310;
 			break;
@@ -1957,98 +2192,267 @@ static bool CompileToGlslWithShaderConductor(
 
 					GlslSource.insert(LayoutPos, DefineString);
 				}
-				for (auto const& Pair : UniformVarNames)
+				if(!bEmulatedUBs)
 				{
-					std::string DefineString = "#define ";
-					DefineString += Pair.first;
-					DefineString += " ";
-					DefineString += Pair.second;
-					DefineString += "\n";
+					for (auto const& Pair : UniformVarNames)
+					{
+						std::string DefineString = "#define ";
+						DefineString += Pair.first;
+						DefineString += " ";
+						DefineString += Pair.second;
+						DefineString += "\n";
 
-					GlslSource.insert(LayoutPos, DefineString);
+						GlslSource.insert(LayoutPos, DefineString);
+					}
 				}
 			}
 
-			std::string GlobalsSearchString = "uniform type_Globals _Globals;";
-			std::string GlobalsString = "//";
 
-			size_t GlobalPos = GlslSource.find(GlobalsSearchString);
-			if (GlobalPos != std::string::npos)
+			// Perform FBF replacements
 			{
-				GlslSource.insert(GlobalPos, GlobalsString);
-				
-				bool UsesFramebufferFetch = Frequency == HSF_PixelShader && GlslSource.find("_Globals.ARM_shader_framebuffer_fetch") != std::string::npos;
 
-				std::string GlobalVarString = "_Globals.";
-				size_t GlobalVarPos = 0;
-				do
+				size_t MainPos = GlslSource.find("#version 310 es");
+				if (MainPos != std::string::npos)
+				{
+					MainPos += strlen("#version 310 es");
+					GlslSource.insert(MainPos, GlslFrameBufferDefines);
+					GlslSource.insert(MainPos, GlslFrameBufferExtensions);
+				}
+
+				// Framebuffer Depth Fetch
+				{
+					std::string FBDFString = "_Globals.ARM_shader_framebuffer_fetch_depth_stencil";
+					std::string FBDFReplaceString = "_Globals_ARM_shader_framebuffer_fetch_depth_stencil";
+
+					size_t FramebufferDepthFetchPos = GlslSource.find(FBDFString);
+					bool UsesFramebufferDepthFetch = Frequency == HSF_PixelShader && GlslSource.find("_Globals.ARM_shader_framebuffer_fetch_depth_stencil") != std::string::npos;
+
+					if (UsesFramebufferDepthFetch)
+					{
+						GlslSource.erase(FramebufferDepthFetchPos, FBDFString.length());
+						GlslSource.insert(FramebufferDepthFetchPos, FBDFReplaceString);
+
+						std::string LastFragDepthARMString = "_Globals._RESERVED_IDENTIFIER_FIXUP_gl_LastFragDepthARM";
+						std::string LastFragDepthARMReplaceString = "gl_LastFragDepthARM";
+
+						size_t LastFragDepthARMStringPos = GlslSource.find(LastFragDepthARMString);
+
+						GlslSource.erase(LastFragDepthARMStringPos, LastFragDepthARMString.length());
+						GlslSource.insert(LastFragDepthARMStringPos, LastFragDepthARMReplaceString);
+					}
+				}
+			}
+
+			if (bEmulatedUBs)
+			{
 				{
-					GlobalVarPos = GlslSource.find(GlobalVarString, GlobalVarPos);
-					if (GlobalVarPos != std::string::npos)
+					bool bIsLayout = true;
+
+					std::string GlobalsSearchString = "layout(std140) uniform type_Globals";
+
+					size_t GlobalPos = GlslSource.find(GlobalsSearchString);
+
+					if (GlobalPos == std::string::npos)
+					{
+						GlobalsSearchString = "struct type_Globals";
+						GlobalPos = GlslSource.find(GlobalsSearchString);
+						bIsLayout = false;
+					}
+
+					if (GlobalPos != std::string::npos)
 					{
-						GlslSource.replace(GlobalVarPos, GlobalVarString.length(), "_Globals_");
-						for (std::string const& SearchString : GlobalArrays)
+						std::string GlobalsEndSearchString;
+						if (bIsLayout)
+						{
+							GlobalsEndSearchString = "_Globals;";
+						}
+						else
+						{
+							GlobalsEndSearchString = "uniform type_Globals _Globals;";
+						}
+
+						size_t GlobalEndPos = GlslSource.find(GlobalsEndSearchString);
+
+						if (GlobalEndPos != std::string::npos)
 						{
-							if (!GlslSource.compare(GlobalVarPos, SearchString.length(), SearchString))
+							GlslSource.erase(GlobalPos, GlobalEndPos - GlobalPos + GlobalsEndSearchString.length());
+
+							bool UsesFramebufferFetch = Frequency == HSF_PixelShader && GlslSource.find("_Globals.ARM_shader_framebuffer_fetch") != std::string::npos;
+
+							std::string GlobalVarString = "_Globals.";
+							size_t GlobalVarPos = 0;
+							do
 							{
-								GlslSource.replace(GlobalVarPos + SearchString.length(), 1, "(");
+								GlobalVarPos = GlslSource.find(GlobalVarString, GlobalVarPos);
+								if (GlobalVarPos != std::string::npos)
+								{
+									GlslSource.replace(GlobalVarPos, GlobalVarString.length(), "_Globals_");
+									for (std::string const& SearchString : GlobalArrays)
+									{
+										if (!GlslSource.compare(GlobalVarPos, SearchString.length(), SearchString))
+										{
+											GlslSource.replace(GlobalVarPos + SearchString.length(), 1, "(");
+
+											size_t ClosingBrace = GlslSource.find("]", GlobalVarPos + SearchString.length());
+											if (ClosingBrace != std::string::npos)
+												GlslSource.replace(ClosingBrace, 1, ")");
+										}
+									}
+								}
+							} while (GlobalVarPos != std::string::npos);
 
-								size_t ClosingBrace = GlslSource.find("]", GlobalVarPos + SearchString.length());
-								if (ClosingBrace != std::string::npos)
-									GlslSource.replace(ClosingBrace, 1, ")");
+							for (auto const& Pair : GlobalOffsets)
+							{
+								if (Pair.Value > 0)
+								{
+									std::string NewUniforms;
+									if (Pair.Key == TEXT("u"))
+									{
+										NewUniforms = "uniform uvec4 ";
+										NewUniforms += FrequencyPrefix;
+										NewUniforms += "u_u[";
+										NewUniforms += std::to_string(Pair.Value);
+										NewUniforms += "];\n";
+									}
+									else if (Pair.Key == TEXT("i"))
+									{
+										NewUniforms = "uniform ivec4 ";
+										NewUniforms += FrequencyPrefix;
+										NewUniforms += "u_i[";
+										NewUniforms += std::to_string(Pair.Value);
+										NewUniforms += "];\n";
+									}
+									else if (Pair.Key == TEXT("h"))
+									{
+										NewUniforms = "uniform highp vec4 ";
+										NewUniforms += FrequencyPrefix;
+										NewUniforms += "u_h[";
+										NewUniforms += std::to_string(Pair.Value);
+										NewUniforms += "];\n";
+									}
+									GlslSource.insert(GlobalPos, NewUniforms);
+								}
+							}
+
+							for (std::string const& Define : GlobalRemap)
+							{
+								GlslSource.insert(GlobalPos, Define);
+							}
+
+							if (UsesFramebufferFetch)
+							{
+								size_t MainPos = GlslSource.find("struct type_Globals");
+								if (MainPos != std::string::npos)
+									GlslSource.insert(MainPos, GlslFrameBufferDefines);
+
+								size_t OutColor = GlslSource.find("0) out ");
+								if (OutColor != std::string::npos)
+									GlslSource.replace(OutColor, 7, "0) FRAME_BUFFERFETCH_STORAGE_QUALIFIER ");
 							}
 						}
 					}
-				} while (GlobalVarPos != std::string::npos);
+				}
 
-				for (auto const& Pair : GlobalOffsets)
+				for (const auto& PackedUBNamePair : PackedUBNames)
 				{
-					if (Pair.Value > 0)
+					bool bIsLayout = true;
+					std::string UBSearchString = "layout(std140) uniform type_" + PackedUBNamePair.Value;
+
+					size_t UBPos = GlslSource.find(UBSearchString);
+
+					if (UBPos == std::string::npos)
 					{
-						std::string NewUniforms;
-						if (Pair.Key == TEXT("u"))
+						UBSearchString = "struct type_" + PackedUBNamePair.Value;
+						UBPos = GlslSource.find(UBSearchString);
+						bIsLayout = false;
+					}
+
+					if (UBPos != std::string::npos)
+					{
+						std::string UBEndSearchString;
+						if (bIsLayout)
 						{
-							NewUniforms = "uniform uvec4 ";
-							NewUniforms += FrequencyPrefix;
-							NewUniforms += "u_u[";
-							NewUniforms += std::to_string(Pair.Value);
-							NewUniforms += "];\n";
+							UBEndSearchString = "} " + PackedUBNamePair.Value + ";";
 						}
-						else if (Pair.Key == TEXT("i"))
+						else
 						{
-							NewUniforms = "uniform ivec4 ";
-							NewUniforms += FrequencyPrefix;
-							NewUniforms += "u_i[";
-							NewUniforms += std::to_string(Pair.Value);
-							NewUniforms += "];\n";
+							UBEndSearchString = "uniform type_" + PackedUBNamePair.Value + " " + PackedUBNamePair.Value + ";";
 						}
-						else if (Pair.Key == TEXT("h"))
+
+						size_t UBEndPos = GlslSource.find(UBEndSearchString);
+
+						if (UBEndPos != std::string::npos)
 						{
-							NewUniforms = "uniform vec4 ";
-							NewUniforms += FrequencyPrefix;
-							NewUniforms += "u_h[";
-							NewUniforms += std::to_string(Pair.Value);
-							NewUniforms += "];\n";
+							GlslSource.erase(UBPos, UBEndPos - UBPos + UBEndSearchString.length());
+
+							size_t UBVarPos = 0;
+
+							for (const PackedUBMemberInfo& MemberInfo : PackedUBMemberInfos[PackedUBNamePair.Key])
+							{
+								std::string UBVarString = PackedUBNamePair.Value + "." + MemberInfo.Name;
+
+								UBVarPos = GlslSource.find(UBVarString);
+								while (UBVarPos != std::string::npos)
+								{
+									GlslSource.erase(UBVarPos, PackedUBNamePair.Value.length() + 1);
+
+									for (std::string const& SearchString : PackedUBArrays[PackedUBNamePair.Key])
+									{
+										if (!GlslSource.compare(UBVarPos, SearchString.length(), SearchString))
+										{
+											GlslSource.replace(UBVarPos + SearchString.length(), 1, "(");
+
+											size_t ClosingBrace = GlslSource.find("]", UBVarPos + SearchString.length());
+											if (ClosingBrace != std::string::npos)
+												GlslSource.replace(ClosingBrace, 1, ")");
+										}
+									}
+
+									UBVarPos = GlslSource.find(UBVarString);
+								}
+							}
+
+							for (auto const& Pair : PackedUBOffsets[PackedUBNamePair.Key])
+							{
+								if (Pair.Value > 0)
+								{
+									std::string NewUniforms;
+									std::string UniformPrefix = FrequencyPrefix + std::string("c") + std::to_string(PackedUBNamePair.Key);
+									if (Pair.Key == TEXT("u"))
+									{
+										NewUniforms = "uniform uvec4 ";
+										NewUniforms += UniformPrefix;
+										NewUniforms += "_u[";
+										NewUniforms += std::to_string(Pair.Value);
+										NewUniforms += "];\n";
+									}
+									else if (Pair.Key == TEXT("i"))
+									{
+										NewUniforms = "uniform ivec4 ";
+										NewUniforms += UniformPrefix;
+										NewUniforms += "_i[";
+										NewUniforms += std::to_string(Pair.Value);
+										NewUniforms += "];\n";
+									}
+									else if (Pair.Key == TEXT("h"))
+									{
+										NewUniforms = "uniform highp vec4 ";
+										NewUniforms += UniformPrefix;
+										NewUniforms += "_h[";
+										NewUniforms += std::to_string(Pair.Value);
+										NewUniforms += "];\n";
+									}
+									GlslSource.insert(UBPos, NewUniforms);
+								}
+							}
+
+							for (std::string const& Define : PackedUBRemap[PackedUBNamePair.Key])
+							{
+								GlslSource.insert(UBPos, Define);
+							}
 						}
-						GlslSource.insert(GlobalPos, NewUniforms);
 					}
 				}
-
-				for (std::string const& Define : GlobalRemap)
-				{
-					GlslSource.insert(GlobalPos, Define);
-				}
-				
-				if (UsesFramebufferFetch)
-				{
-					size_t MainPos = GlslSource.find("struct type_Globals");
-					if (MainPos != std::string::npos)
-						GlslSource.insert(MainPos, GlslFrameBufferDefines);
-					
-					size_t OutColor = GlslSource.find("0) out ");
-					if (OutColor != std::string::npos)
-						GlslSource.replace(OutColor, 7, "0) FRAME_BUFFERFETCH_STORAGE_QUALIFIER ");
-				}
 			}
 
 			size_t GlslSourceLen = GlslSource.length();
@@ -2083,11 +2487,49 @@ static bool CompileToGlslWithShaderConductor(
 						}
 					}
 
+					// Rename texture buffers
+					std::string SamplerBufferName = std::string("samplerBuffer ") + TCHAR_TO_ANSI(*Texture);
+					size_t SamplerBufferPos = GlslSource.find(SamplerBufferName);
+
+					if (SamplerBufferPos != std::string::npos)
+					{
+						std::string NewSamplerBufferName = std::string(FrequencyPrefix) + "s" + std::to_string(TextureIndex);
+
+						GlslSource.erase(SamplerBufferPos + SamplerBufferName.size() - Texture.Len(), Texture.Len());
+						GlslSource.insert(SamplerBufferPos + SamplerBufferName.size() - Texture.Len(), NewSamplerBufferName);
+
+						std::string SamplerTexFetchName = std::string("texelFetch(") + TCHAR_TO_ANSI(*Texture);
+						size_t SamplerTexFetchPos = GlslSource.find(SamplerTexFetchName);
+
+						while (SamplerTexFetchPos != std::string::npos)
+						{
+
+							GlslSource.erase(SamplerTexFetchPos + SamplerTexFetchName.size() - Texture.Len(), Texture.Len());
+							GlslSource.insert(SamplerTexFetchPos + SamplerTexFetchName.size() - Texture.Len(), NewSamplerBufferName);
+
+							SamplerTexFetchPos = GlslSource.find(SamplerTexFetchName);
+						}
+					}
+
 					const uint32 SamplerCount = FMath::Max(1, UsedSamplers.Num());
 					CCHeaderWriter.WriteSRV(*Texture, TextureIndex, SamplerCount, UsedSamplers);
 					TextureIndex += SamplerCount;
 				}
 
+				for (const auto & pair : UniformVarNames)
+				{
+					std::string OldUniformTypeName = "uniform type_" + pair.first;
+					std::string NewUniformTypeName = "uniform " + pair.second;
+
+					size_t UniformTypePos = GlslSource.find(OldUniformTypeName);
+
+					if (UniformTypePos != std::string::npos)
+					{
+						GlslSource.erase(UniformTypePos, OldUniformTypeName.length());
+						GlslSource.insert(UniformTypePos, NewUniformTypeName);
+					}
+				}
+
 				// Generate meta data for CCHeader
 				CCHeaderWriter.WriteSourceInfo(*Input.GetSourceFilename(), *Input.EntryPointName, *Input.DebugGroupName);
 				CCHeaderWriter.WriteCompilerInfo();
diff --git a/Engine/Source/Developer/ShaderFormatOpenGL/Private/ShaderFormatOpenGL.cpp b/Engine/Source/Developer/ShaderFormatOpenGL/Private/ShaderFormatOpenGL.cpp
--- a/Engine/Source/Developer/ShaderFormatOpenGL/Private/ShaderFormatOpenGL.cpp
+++ b/Engine/Source/Developer/ShaderFormatOpenGL/Private/ShaderFormatOpenGL.cpp
@@ -77,7 +77,7 @@ public:
 
 	virtual const TCHAR* GetPlatformIncludeDirectory() const
 	{
-		return TEXT("");
+		return TEXT("GL");
 	}
 
 	virtual bool UsesHLSLcc(const struct FShaderCompilerInput& Input) const override
diff --git a/Engine/Source/Runtime/Engine/Private/ShaderCompiler/ShaderCompiler.cpp b/Engine/Source/Runtime/Engine/Private/ShaderCompiler/ShaderCompiler.cpp
--- a/Engine/Source/Runtime/Engine/Private/ShaderCompiler/ShaderCompiler.cpp
+++ b/Engine/Source/Runtime/Engine/Private/ShaderCompiler/ShaderCompiler.cpp
@@ -936,10 +936,10 @@ static TAutoConsoleVariable<int32> CVarD3DForceShaderConductorDXCRewrite(
 
 static TAutoConsoleVariable<int32> CVarOpenGLForceDXC(
 	TEXT("r.OpenGL.ForceDXC"),
-	0,
+	1,
 	TEXT("Forces DirectX Shader Compiler (DXC) to be used for all OpenGL shaders instead of hlslcc.\n")
-	TEXT(" 0: Disable (default)\n")
-	TEXT(" 1: Force new compiler for all shaders"),
+	TEXT(" 0: Disable\n")
+	TEXT(" 1: Force new compiler for all shaders (default)"),
 	ECVF_ReadOnly);
 
 static TAutoConsoleVariable<int32> CVarVulkanForceDXC(
@@ -5107,13 +5107,9 @@ void GlobalBeginCompileShader(
 
 	if (bIsMobilePlatform)
 	{
-		if (IsOpenGLPlatform((EShaderPlatform)Target.Platform))
+		if (IsUsingEmulatedUniformBuffers((EShaderPlatform)Target.Platform))
 		{
-			static const auto CVar = IConsoleManager::Get().FindConsoleVariable(TEXT("OpenGL.UseEmulatedUBs"));
-			if (CVar && CVar->GetInt() != 0)
-			{
-				Input.Environment.CompilerFlags.Add(CFLAG_UseEmulatedUB);
-			}
+			Input.Environment.CompilerFlags.Add(CFLAG_UseEmulatedUB);
 		}
 	}
 
diff --git a/Engine/Source/Runtime/OpenGLDrv/Private/OpenGLDevice.cpp b/Engine/Source/Runtime/OpenGLDrv/Private/OpenGLDevice.cpp
--- a/Engine/Source/Runtime/OpenGLDrv/Private/OpenGLDevice.cpp
+++ b/Engine/Source/Runtime/OpenGLDrv/Private/OpenGLDevice.cpp
@@ -803,9 +803,7 @@ static void InitRHICapabilitiesForGL()
 		;
 	
 	// By default use emulated UBs on mobile
-	static auto* CVar = IConsoleManager::Get().FindTConsoleVariableDataInt(TEXT("OpenGL.UseEmulatedUBs"));
-	const bool bUseEmulatedUBs = (CVar && CVar->GetValueOnAnyThread() != 0);
-	GUseEmulatedUniformBuffers = (IsMobilePlatform(GMaxRHIShaderPlatform) && bUseEmulatedUBs);
+	GUseEmulatedUniformBuffers = IsUsingEmulatedUniformBuffers(GMaxRHIShaderPlatform);
 
 	FString FeatureLevelName;
 	GetFeatureLevelName(GMaxRHIFeatureLevel, FeatureLevelName);
diff --git a/Engine/Source/Runtime/RenderCore/Private/Shader.cpp b/Engine/Source/Runtime/RenderCore/Private/Shader.cpp
--- a/Engine/Source/Runtime/RenderCore/Private/Shader.cpp
+++ b/Engine/Source/Runtime/RenderCore/Private/Shader.cpp
@@ -1316,6 +1316,24 @@ bool IsDxcEnabledForPlatform(EShaderPlatform Platform)
 	return false;
 }
 
+bool IsUsingEmulatedUniformBuffers(EShaderPlatform Platform)
+{
+	if (IsAndroidOpenGLESPlatform(Platform))
+	{
+		// Currently DXC only supports emulated uniform buffers on GLES
+		static const auto CForceDXCVar = IConsoleManager::Get().FindConsoleVariable(TEXT("r.OpenGL.ForceDXC"));
+		if (CForceDXCVar && CForceDXCVar->GetInt() != 0)
+		{
+			return true;
+		}
+
+		static auto* CVar = IConsoleManager::Get().FindTConsoleVariableDataInt(TEXT("OpenGL.UseEmulatedUBs"));
+		return (CVar && CVar->GetValueOnAnyThread() != 0);
+	}
+
+	return false;
+}
+
 void ShaderMapAppendKeyString(EShaderPlatform Platform, FString& KeyString)
 {
 	const FName ShaderFormatName = LegacyShaderPlatformToShaderFormat(Platform);
@@ -1496,11 +1514,7 @@ void ShaderMapAppendKeyString(EShaderPlatform Platform, FString& KeyString)
 			KeyString += (CVar && CVar->GetInt() != 0) ? TEXT("_DLODT") : TEXT("");
 		}
 		
-		if (IsOpenGLPlatform(Platform))
-		{
-			static IConsoleVariable* CVar = IConsoleManager::Get().FindConsoleVariable(TEXT("OpenGL.UseEmulatedUBs"));
-			KeyString += (CVar && CVar->GetInt() != 0) ? TEXT("_NoUB") : TEXT("");
-		}
+		KeyString += IsUsingEmulatedUniformBuffers(Platform) ? TEXT("_NoUB") : TEXT("");
 
 		{
 			static FShaderPlatformCachedIniValue<bool> MobileEnableMovableSpotlightsIniValue(TEXT("/Script/Engine.RendererSettings"), TEXT("r.Mobile.EnableMovableSpotlights"));
@@ -1549,11 +1563,7 @@ void ShaderMapAppendKeyString(EShaderPlatform Platform, FString& KeyString)
 	}
 	else
 	{
-		if (IsOpenGLPlatform(Platform))
-		{
-			static IConsoleVariable* CVar = IConsoleManager::Get().FindConsoleVariable(TEXT("OpenGL.UseEmulatedUBs"));
-			KeyString += (CVar && CVar->GetInt() != 0) ? TEXT("_NoUB") : TEXT("");
-		}
+		KeyString += IsUsingEmulatedUniformBuffers(Platform) ? TEXT("_NoUB") : TEXT("");
 	}
 
 	const IShaderFormat* ShaderFormat = GetTargetPlatformManagerRef().FindShaderFormat(ShaderFormatName);
