Fixed for Linux compile errors

diff --git a/Engine/Source/Developer/ShaderFormatOpenGL/Private/OpenGLShaderCompiler.cpp b/Engine/Source/Developer/ShaderFormatOpenGL/Private/OpenGLShaderCompiler.cpp
--- a/Engine/Source/Developer/ShaderFormatOpenGL/Private/OpenGLShaderCompiler.cpp
+++ b/Engine/Source/Developer/ShaderFormatOpenGL/Private/OpenGLShaderCompiler.cpp
@@ -1255,13 +1255,13 @@ void WritePackedUBHeader(CrossCompiler::FHlslccHeaderWriter& CCHeaderWriter, con
 
 	bool bNeedsHeader = true;
 
-	for (const auto Pair : UBNames)
+	for (const auto & Pair : UBNames)
 	{
 		FString Name(Pair.Value.c_str());
 		CCHeaderWriter.WritePackedUB(Name, Pair.Key);
 	}
 
-	for (const auto Pair : UBMemberInfo)
+	for (const auto & Pair : UBMemberInfo)
 	{
 		FString UBName(UBNames[Pair.Key].c_str());
 		
@@ -1382,7 +1382,7 @@ void AddMemberToPackedUB(const std::string & FrequencyPrefix,
 		else
 		{
 			char buff[10];
-			sprintf_s(buff, 10, "mat%dx%d(", type.traits.numeric.matrix.column_count, type.traits.numeric.matrix.row_count);
+			sprintf(buff, "mat%dx%d(", type.traits.numeric.matrix.column_count, type.traits.numeric.matrix.row_count);
 			Name += buff;
 		}
 
@@ -1398,7 +1398,7 @@ void AddMemberToPackedUB(const std::string & FrequencyPrefix,
 			Name += TCHAR_TO_UTF8(*TypeQualifier);
 			
 			char buff[10];
-			sprintf_s(buff, 10, "[%s + %d]", OffsetString.c_str(), i);
+			sprintf(buff, "[%s + %d]", OffsetString.c_str(), i);
 
 			Name += buff;
 
diff --git a/Engine/Source/Runtime/Engine/Private/ShaderCompiler/ShaderCompiler.cpp b/Engine/Source/Runtime/Engine/Private/ShaderCompiler/ShaderCompiler.cpp
--- a/Engine/Source/Runtime/Engine/Private/ShaderCompiler/ShaderCompiler.cpp
+++ b/Engine/Source/Runtime/Engine/Private/ShaderCompiler/ShaderCompiler.cpp
@@ -53,6 +53,7 @@
 #include "ProfilingDebugging/DiagnosticTable.h"
 #include "ProfilingDebugging/LoadTimeTracker.h"
 #include "ShaderCore.h"
+#include "Shader.h"
 #include "DistributedBuildInterface/Public/DistributedBuildControllerInterface.h"
 #include "Misc/ScopeRWLock.h"
 #include "Async/ParallelFor.h"
diff --git a/Engine/Source/Runtime/OpenGLDrv/Private/OpenGLDevice.cpp b/Engine/Source/Runtime/OpenGLDrv/Private/OpenGLDevice.cpp
--- a/Engine/Source/Runtime/OpenGLDrv/Private/OpenGLDevice.cpp
+++ b/Engine/Source/Runtime/OpenGLDrv/Private/OpenGLDevice.cpp
@@ -16,6 +16,7 @@
 #include "Containers/List.h"
 #include "RenderResource.h"
 #include "ShaderCore.h"
+#include "Shader.h"
 #include "RenderUtils.h"
 #include "OpenGLDrv.h"
 #include "OpenGLDrvPrivate.h"
diff --git a/Engine/Source/Runtime/RenderCore/Public/Shader.h b/Engine/Source/Runtime/RenderCore/Public/Shader.h
--- a/Engine/Source/Runtime/RenderCore/Public/Shader.h
+++ b/Engine/Source/Runtime/RenderCore/Public/Shader.h
@@ -2347,6 +2347,9 @@ inline void DispatchComputeShader(
 	DispatchComputeShader(RHICmdList, Shader.GetShader(), ThreadGroupCountX, ThreadGroupCountY, ThreadGroupCountZ);
 }
 
+/** Returns whether the platform is using emulated uniform buffers */
+extern RENDERCORE_API bool IsUsingEmulatedUniformBuffers(EShaderPlatform Platform);
+
 /** Returns whether DirectXShaderCompiler (DXC) is enabled for the specified shader platform. See console variables "r.OpenGL.ForceDXC", "r.Vulkan.ForceDXC", and "r.D3D.ForceDXC". */
 extern RENDERCORE_API bool IsDxcEnabledForPlatform(EShaderPlatform Platform);
 
