Removed sprintf

diff --git a/Engine/Source/Developer/ShaderFormatOpenGL/Private/OpenGLShaderCompiler.cpp b/Engine/Source/Developer/ShaderFormatOpenGL/Private/OpenGLShaderCompiler.cpp
--- a/Engine/Source/Developer/ShaderFormatOpenGL/Private/OpenGLShaderCompiler.cpp
+++ b/Engine/Source/Developer/ShaderFormatOpenGL/Private/OpenGLShaderCompiler.cpp
@@ -1381,9 +1381,8 @@ void AddMemberToPackedUB(const std::string & FrequencyPrefix,
 		}
 		else
 		{
-			char buff[10];
-			sprintf(buff, "mat%dx%d(", type.traits.numeric.matrix.column_count, type.traits.numeric.matrix.row_count);
-			Name += buff;
+			std::string Buff = "mat" + std::to_string(type.traits.numeric.matrix.column_count) + "x" + std::to_string(type.traits.numeric.matrix.row_count) + "(";
+			Name += Buff;
 		}
 
 		for (uint32_t i = 0; i < type.traits.numeric.matrix.row_count; ++i)
@@ -1396,11 +1395,9 @@ void AddMemberToPackedUB(const std::string & FrequencyPrefix,
 			Name += UniformPrefix;
 			Name += "_";
 			Name += TCHAR_TO_UTF8(*TypeQualifier);
-			
-			char buff[10];
-			sprintf(buff, "[%s + %d]", OffsetString.c_str(), i);
 
-			Name += buff;
+			std::string Buff = "[" + OffsetString + " + " + std::to_string(i) + "]";
+			Name += Buff;
 
 			switch (type.traits.numeric.matrix.column_count)
 			{
