Replace EHlslShaderFrequency with EShaderFrequency in ShaderConductor interface.
Also remove further references to Hull and Domain shaders as tessellation is no longer supported in UE5.

#rnx

diff --git a/Engine/Source/Developer/Apple/MetalShaderFormat/Private/MetalDerivedData.cpp b/Engine/Source/Developer/Apple/MetalShaderFormat/Private/MetalDerivedData.cpp
--- a/Engine/Source/Developer/Apple/MetalShaderFormat/Private/MetalDerivedData.cpp
+++ b/Engine/Source/Developer/Apple/MetalShaderFormat/Private/MetalDerivedData.cpp
@@ -163,7 +163,7 @@ bool FMetalShaderBytecodeCooker::Build(TArray<uint8>& OutData)
 	return bSuccess;
 }
 
-FMetalShaderOutputCooker::FMetalShaderOutputCooker(const FShaderCompilerInput& _Input, FShaderCompilerOutput& _Output, const FString& _WorkingDirectory, FString _PreprocessedShader, FSHAHash _GUIDHash, uint8 _VersionEnum, uint32 _CCFlags, EHlslCompileTarget _HlslCompilerTarget, EHlslCompileTarget _MetalCompilerTarget, EMetalGPUSemantics _Semantics, EMetalTypeBufferMode _TypeMode, uint32 _MaxUnrollLoops, EHlslShaderFrequency _Frequency, bool _bDumpDebugInfo, FString _Standard, FString _MinOSVersion)
+FMetalShaderOutputCooker::FMetalShaderOutputCooker(const FShaderCompilerInput& _Input, FShaderCompilerOutput& _Output, const FString& _WorkingDirectory, FString _PreprocessedShader, FSHAHash _GUIDHash, uint8 _VersionEnum, uint32 _CCFlags, EHlslCompileTarget _HlslCompilerTarget, EHlslCompileTarget _MetalCompilerTarget, EMetalGPUSemantics _Semantics, EMetalTypeBufferMode _TypeMode, uint32 _MaxUnrollLoops, EShaderFrequency _Frequency, bool _bDumpDebugInfo, FString _Standard, FString _MinOSVersion)
 	: Input(_Input)
 	, Output(_Output)
 	, WorkingDirectory(_WorkingDirectory)
@@ -381,7 +381,7 @@ bool FMetalShaderOutputCooker::Build(TArray<uint8>& OutData)
 
 		// Always disable FMA pass for Pixel and Compute shader,
 		// otherwise determine whether [[position, invariant]] qualifier is available in Metal or not.
-		if (Frequency == HSF_PixelShader || Frequency == HSF_ComputeShader)
+		if (Frequency == SF_Pixel || Frequency == SF_Compute)
 		{
 			Options.bEnableFMAPass = false;
 		}
@@ -423,7 +423,6 @@ bool FMetalShaderOutputCooker::Build(TArray<uint8>& OutData)
 		
 		CrossCompiler::FHlslccHeaderWriter CCHeaderWriter;
 
-		EHlslShaderFrequency Freq = Frequency;
 		FString ALNString;
 		uint32 IABOffsetIndex = 0;
 		uint64 BufferIndices = 0xffffffffffffffff;
@@ -542,7 +541,7 @@ bool FMetalShaderOutputCooker::Build(TArray<uint8>& OutData)
 					switch (Mode->mode) {
 						case SpvExecutionModeLocalSize:
 						case SpvExecutionModeLocalSizeHint:
-							if (Frequency == HSF_ComputeShader)
+							if (Frequency == SF_Compute)
 							{
 								check(Mode->operands_count == 3);
 								CCHeaderWriter.WriteNumThreads(Mode->operands[0], Mode->operands[1], Mode->operands[2]);
@@ -923,7 +922,7 @@ bool FMetalShaderOutputCooker::Build(TArray<uint8>& OutData)
 				}
 			}
 			
-			if (Frequency == HSF_PixelShader)
+			if (Frequency == SF_Pixel)
 			{
 				ReflectionBindings.GatherOutputAttributes(Reflection);
 				for (auto const& Var : ReflectionBindings.OutputAttributes)
@@ -960,7 +959,7 @@ bool FMetalShaderOutputCooker::Build(TArray<uint8>& OutData)
 				}
 			}
 			
-			if (Frequency == HSF_VertexShader)
+			if (Frequency == SF_Vertex)
 			{
 				uint32 AssignedInputs = 0;
 
diff --git a/Engine/Source/Developer/Apple/MetalShaderFormat/Private/MetalDerivedData.h b/Engine/Source/Developer/Apple/MetalShaderFormat/Private/MetalDerivedData.h
--- a/Engine/Source/Developer/Apple/MetalShaderFormat/Private/MetalDerivedData.h
+++ b/Engine/Source/Developer/Apple/MetalShaderFormat/Private/MetalDerivedData.h
@@ -132,30 +132,10 @@ struct FMetalShaderPreprocessed
 	}
 };
 
-struct FMetalShaderOutputJob
-{
-	const FShaderCompilerInput& Input;
-	FShaderCompilerOutput& Output;
-	const FString& WorkingDirectory;
-	FString PreprocessedShader;
-	FSHAHash GUIDHash;
-	uint8 VersionEnum;
-	uint32 CCFlags;
-	EHlslCompileTarget HlslCompilerTarget;
-	EHlslCompileTarget MetalCompilerTarget;
-	EMetalGPUSemantics Semantics;
-	EMetalTypeBufferMode TypeMode;
-	uint32 MaxUnrollLoops;
-	EHlslShaderFrequency Frequency;
-	bool bDumpDebugInfo;
-	FString Standard;
-	FString MinOSVersion;
-};
-
 class FMetalShaderOutputCooker : public FDerivedDataPluginInterface
 {
 public:
-	FMetalShaderOutputCooker(const FShaderCompilerInput& _Input,FShaderCompilerOutput& Output,const FString& WorkingDirectory, FString PreprocessedShader, FSHAHash GUIDHash, uint8 VersionEnum, uint32 CCFlags, EHlslCompileTarget HlslCompilerTarget, EHlslCompileTarget MetalCompilerTarget, EMetalGPUSemantics Semantics, EMetalTypeBufferMode TypeMode, uint32 MaxUnrollLoops, EHlslShaderFrequency Frequency, bool bDumpDebugInfo, FString Standard, FString MinOSVersion);
+	FMetalShaderOutputCooker(const FShaderCompilerInput& _Input,FShaderCompilerOutput& Output,const FString& WorkingDirectory, FString PreprocessedShader, FSHAHash GUIDHash, uint8 VersionEnum, uint32 CCFlags, EHlslCompileTarget HlslCompilerTarget, EHlslCompileTarget MetalCompilerTarget, EMetalGPUSemantics Semantics, EMetalTypeBufferMode TypeMode, uint32 MaxUnrollLoops, EShaderFrequency Frequency, bool bDumpDebugInfo, FString Standard, FString MinOSVersion);
 	virtual ~FMetalShaderOutputCooker();
 
 #if PLATFORM_MAC || PLATFORM_IOS
@@ -181,7 +161,7 @@ private:
 	EMetalGPUSemantics Semantics;
 	EMetalTypeBufferMode TypeMode;
 	uint32 MaxUnrollLoops;
-	EHlslShaderFrequency Frequency;
+	EShaderFrequency Frequency;
 	bool bDumpDebugInfo;
 	FString Standard;
 	FString MinOSVersion;
diff --git a/Engine/Source/Developer/Apple/MetalShaderFormat/Private/MetalShaderCompiler.cpp b/Engine/Source/Developer/Apple/MetalShaderFormat/Private/MetalShaderCompiler.cpp
--- a/Engine/Source/Developer/Apple/MetalShaderFormat/Private/MetalShaderCompiler.cpp
+++ b/Engine/Source/Developer/Apple/MetalShaderFormat/Private/MetalShaderCompiler.cpp
@@ -807,17 +807,6 @@ void BuildMetalShaderOutput(
 	External interface.
 ------------------------------------------------------------------------------*/
 
-static const EHlslShaderFrequency FrequencyTable[] =
-{
-	HSF_VertexShader,
-	HSF_InvalidFrequency,
-	HSF_InvalidFrequency,
-	HSF_PixelShader,
-	HSF_InvalidFrequency,
-	HSF_ComputeShader
-};
-static_assert(SF_NumStandardFrequencies == UE_ARRAY_COUNT(FrequencyTable), "NumFrequencies changed. Please update tables.");
-
 FString CreateRemoteDataFromEnvironment(const FShaderCompilerEnvironment& Environment)
 {
 	FString Line = TEXT("\n#if 0 /*BEGIN_REMOTE_SERVER*/\n");
@@ -1143,14 +1132,14 @@ void CompileShader_Metal(const FShaderCompilerInput& _Input,FShaderCompilerOutpu
 	char* MetalShaderSource = NULL;
 	char* ErrorLog = NULL;
 
-	const EHlslShaderFrequency Frequency = FrequencyTable[Input.Target.Frequency];
-	if (Frequency == HSF_InvalidFrequency)
+	const EShaderFrequency Frequency = Input.Target.Frequency;
+	if (!(Frequency == SF_Vertex || Frequency == SF_Pixel || Frequency == SF_Compute))
 	{
 		Output.bSucceeded = false;
 		FShaderCompilerError* NewError = new(Output.Errors) FShaderCompilerError();
 		NewError->StrippedErrorMessage = FString::Printf(
 			TEXT("%s shaders not supported for use in Metal."),
-			CrossCompiler::GetFrequencyName((EShaderFrequency)Input.Target.Frequency)
+			CrossCompiler::GetFrequencyName(Frequency)
 			);
 		return;
 	}
diff --git a/Engine/Source/Developer/ShaderCompilerCommon/Private/ShaderConductorContext.cpp b/Engine/Source/Developer/ShaderCompilerCommon/Private/ShaderConductorContext.cpp
--- a/Engine/Source/Developer/ShaderCompilerCommon/Private/ShaderConductorContext.cpp
+++ b/Engine/Source/Developer/ShaderCompilerCommon/Private/ShaderConductorContext.cpp
@@ -169,24 +169,26 @@ namespace CrossCompiler
 		return false;
 	}
 
-	static ShaderConductor::ShaderStage ToShaderConductorShaderStage(EHlslShaderFrequency Frequency)
+	static ShaderConductor::ShaderStage ToShaderConductorShaderStage(EShaderFrequency Frequency)
 	{
-		check(Frequency >= HSF_VertexShader && Frequency <= HSF_RayCallable);
+		check(Frequency >= SF_Vertex && Frequency <= SF_RayCallable);
 		switch (Frequency)
 		{
-		case HSF_VertexShader:		return ShaderConductor::ShaderStage::VertexShader;
-		case HSF_PixelShader:		return ShaderConductor::ShaderStage::PixelShader;
-		case HSF_GeometryShader:	return ShaderConductor::ShaderStage::GeometryShader;
-		case HSF_HullShader:		return ShaderConductor::ShaderStage::HullShader;
-		case HSF_DomainShader:		return ShaderConductor::ShaderStage::DomainShader;
-		case HSF_ComputeShader:		return ShaderConductor::ShaderStage::ComputeShader;
+		case SF_Vertex:			return ShaderConductor::ShaderStage::VertexShader;
+		case SF_Mesh:			checkf(0, TEXT("SF_Mesh not support in ShaderConductor")); break;
+		case SF_Amplification:	checkf(0, TEXT("SF_Amplification not support in ShaderConductor")); break;
+		case SF_Pixel:			return ShaderConductor::ShaderStage::PixelShader;
+		case SF_Geometry:		return ShaderConductor::ShaderStage::GeometryShader;
+		case SF_Compute:		return ShaderConductor::ShaderStage::ComputeShader;
 
-		case HSF_RayGen:			return ShaderConductor::ShaderStage::RayGen;
-		case HSF_RayMiss:			return ShaderConductor::ShaderStage::RayMiss;
-		case HSF_RayHitGroup:		return ShaderConductor::ShaderStage::RayHitGroup;
-		case HSF_RayCallable:		return ShaderConductor::ShaderStage::RayCallable;
-		default:					return ShaderConductor::ShaderStage::NumShaderStages;
+		case SF_RayGen:			return ShaderConductor::ShaderStage::RayGen;
+		case SF_RayMiss:		return ShaderConductor::ShaderStage::RayMiss;
+		case SF_RayHitGroup:	return ShaderConductor::ShaderStage::RayHitGroup;
+		case SF_RayCallable:	return ShaderConductor::ShaderStage::RayCallable;
+
+		default: break;
 		}
+		return ShaderConductor::ShaderStage::NumShaderStages;
 	}
 
 	static ShaderConductor::Compiler::ShaderModel ToShaderConductorShaderModel(EHlslCompileTarget Target)
@@ -483,7 +485,7 @@ namespace CrossCompiler
 		return *this;
 	}
 
-	bool FShaderConductorContext::LoadSource(const FString& ShaderSource, const FString& Filename, const FString& EntryPoint, EHlslShaderFrequency ShaderStage, const FShaderCompilerDefinitions* Definitions, const TArray<FString>* ExtraDxcArgs)
+	bool FShaderConductorContext::LoadSource(const FString& ShaderSource, const FString& Filename, const FString& EntryPoint, EShaderFrequency ShaderStage, const FShaderCompilerDefinitions* Definitions, const TArray<FString>* ExtraDxcArgs)
 	{
 		// Convert FString to ANSI string and store them as intermediates
 		ConvertFStringToAnsiString(ShaderSource, Intermediates->ShaderSource);
@@ -507,7 +509,7 @@ namespace CrossCompiler
 		return true;
 	}
 
-	bool FShaderConductorContext::LoadSource(const ANSICHAR* ShaderSource, const ANSICHAR* Filename, const ANSICHAR* EntryPoint, EHlslShaderFrequency ShaderStage, const FShaderCompilerDefinitions* Definitions, const TArray<FString>* ExtraDxcArgs)
+	bool FShaderConductorContext::LoadSource(const ANSICHAR* ShaderSource, const ANSICHAR* Filename, const ANSICHAR* EntryPoint, EShaderFrequency ShaderStage, const FShaderCompilerDefinitions* Definitions, const TArray<FString>* ExtraDxcArgs)
 	{
 		// Store ANSI strings as intermediates
 		CopyAnsiString(ShaderSource, Intermediates->ShaderSource);
@@ -838,12 +840,12 @@ namespace CrossCompiler
 		return *this; // Dummy
 	}
 
-	bool FShaderConductorContext::LoadSource(const FString& ShaderSource, const FString& Filename, const FString& EntryPoint, EHlslShaderFrequency ShaderStage, const FShaderCompilerDefinitions* Definitions, const TArray<FString>* ExtraDxcArgs)
+	bool FShaderConductorContext::LoadSource(const FString& ShaderSource, const FString& Filename, const FString& EntryPoint, EShaderFrequency ShaderStage, const FShaderCompilerDefinitions* Definitions, const TArray<FString>* ExtraDxcArgs)
 	{
 		return false; // Dummy
 	}
 
-	bool FShaderConductorContext::LoadSource(const ANSICHAR* ShaderSource, const ANSICHAR* Filename, const ANSICHAR* EntryPoint, EHlslShaderFrequency ShaderStage, const FShaderCompilerDefinitions* Definitions, const TArray<FString>* ExtraDxcArgs)
+	bool FShaderConductorContext::LoadSource(const ANSICHAR* ShaderSource, const ANSICHAR* Filename, const ANSICHAR* EntryPoint, EShaderFrequency ShaderStage, const FShaderCompilerDefinitions* Definitions, const TArray<FString>* ExtraDxcArgs)
 	{
 		return false; // Dummy
 	}
@@ -917,23 +919,35 @@ namespace CrossCompiler
 		return IdentifierTable;
 	}
 
-	EHlslShaderFrequency FShaderConductorContext::MapShaderFrequency(EShaderFrequency Frequency)
+	static const TCHAR* GetGlslShaderFileExt(EShaderFrequency ShaderStage)
 	{
-		const EHlslShaderFrequency FrequencyTable[] =
+		switch (ShaderStage)
 		{
-			HSF_VertexShader,		// SF_Vertex
-			HSF_InvalidFrequency,	// SF_Mesh
-			HSF_InvalidFrequency,	// SF_Amplification
-			HSF_PixelShader,		// SF_Pixel
-			HSF_GeometryShader,		// SF_Geometry
-			HSF_ComputeShader,		// SF_Compute
-			HSF_RayGen,				// SF_RayGen
-			HSF_RayMiss,			// SF_RayMiss
-			HSF_RayHitGroup,		// SF_RayHitGroup
-			HSF_RayCallable,		// SF_RayCallable
-		};
-		const uint32 FrequencyIndex = static_cast<uint32>(Frequency);
-		return (FrequencyIndex < UE_ARRAY_COUNT(FrequencyTable) ? FrequencyTable[FrequencyIndex] : HSF_InvalidFrequency);
+		case SF_Vertex:			return TEXT("vert");
+		case SF_Mesh:			return TEXT("mesh");
+		case SF_Amplification:	return TEXT("task");
+		case SF_Pixel:			return TEXT("frag");
+		case SF_Geometry:		return TEXT("geom");
+		case SF_Compute:		return TEXT("comp");
+		case SF_RayGen:			return TEXT("rgen");
+		case SF_RayMiss:		return TEXT("rmiss");
+		case SF_RayHitGroup:	return TEXT("rahit"); // rahit/rchit
+		case SF_RayCallable:	return TEXT("rcall");
+		default:				return TEXT("glsl");
+		}
+	}
+
+	const TCHAR* FShaderConductorContext::GetShaderFileExt(EShaderConductorLanguage Language, EShaderFrequency ShaderStage)
+	{
+		switch (Language)
+		{
+		case EShaderConductorLanguage::Hlsl:		return TEXT("hlsl");
+		case EShaderConductorLanguage::Glsl:		[[fallthrough]];
+		case EShaderConductorLanguage::Essl:		return GetGlslShaderFileExt(ShaderStage);
+		case EShaderConductorLanguage::Metal_macOS: [[fallthrough]];
+		case EShaderConductorLanguage::Metal_iOS:	return TEXT("metal");
+		default:									return TEXT("");
+		}
 	}
 
 } // namespace CrossCompiler
diff --git a/Engine/Source/Developer/ShaderCompilerCommon/Public/ShaderConductorContext.h b/Engine/Source/Developer/ShaderCompilerCommon/Public/ShaderConductorContext.h
--- a/Engine/Source/Developer/ShaderCompilerCommon/Public/ShaderConductorContext.h
+++ b/Engine/Source/Developer/ShaderCompilerCommon/Public/ShaderConductorContext.h
@@ -107,8 +107,8 @@ namespace CrossCompiler
 		FShaderConductorContext& operator = (const FShaderConductorContext&) = delete;
 
 		/** Loads the shader source and converts the input descriptor to a format suitable for ShaderConductor. If 'Definitions' is null, the previously loaded definitions are not modified. */
-		bool LoadSource(const FString& ShaderSource, const FString& Filename, const FString& EntryPoint, EHlslShaderFrequency ShaderStage, const FShaderCompilerDefinitions* Definitions = nullptr, const TArray<FString>* ExtraDxcArgs = nullptr);
-		bool LoadSource(const ANSICHAR* ShaderSource, const ANSICHAR* Filename, const ANSICHAR* EntryPoint, EHlslShaderFrequency ShaderStage, const FShaderCompilerDefinitions* Definitions = nullptr, const TArray<FString>* ExtraDxcArgs = nullptr);
+		bool LoadSource(const FString& ShaderSource, const FString& Filename, const FString& EntryPoint, EShaderFrequency ShaderStage, const FShaderCompilerDefinitions* Definitions = nullptr, const TArray<FString>* ExtraDxcArgs = nullptr);
+		bool LoadSource(const ANSICHAR* ShaderSource, const ANSICHAR* Filename, const ANSICHAR* EntryPoint, EShaderFrequency ShaderStage, const FShaderCompilerDefinitions* Definitions = nullptr, const TArray<FString>* ExtraDxcArgs = nullptr);
 
 		/** Rewrites the specified HLSL shader source code. This allows to reduce the HLSL code by removing unused global resources for instance.
 		This will update the internally loaded source (see 'LoadSource'), so the output parameter 'OutSource' is optional. */
@@ -158,8 +158,8 @@ namespace CrossCompiler
 		/** Disassembles the specified SPIR-V module and returns its assembly as text representation. */
 		static bool Disassemble(EShaderConductorIR Language, const void* Binary, uint32 BinaryByteSize, TArray<ANSICHAR>& OutAssemblyText);
 
-		/** Maps the primary shader frequency (see FShaderTarget::Frequency) to EHlslShaderFrequency used by this context. */
-		static EHlslShaderFrequency MapShaderFrequency(EShaderFrequency Frequency);
+		/** Returns a filename extension for the specified shading language and shader stage, e.g. "frag" for a GLSL pixel shader. */
+		static const TCHAR* GetShaderFileExt(EShaderConductorLanguage Language, EShaderFrequency ShaderStage = SF_NumFrequencies);
 
 	public:
 		struct FShaderConductorIntermediates; // Pimpl idiom
diff --git a/Engine/Source/Developer/ShaderFormatOpenGL/Private/GlslBackend.cpp b/Engine/Source/Developer/ShaderFormatOpenGL/Private/GlslBackend.cpp
--- a/Engine/Source/Developer/ShaderFormatOpenGL/Private/GlslBackend.cpp
+++ b/Engine/Source/Developer/ShaderFormatOpenGL/Private/GlslBackend.cpp
@@ -4005,21 +4005,9 @@ static FSystemValue GeometrySystemValueTable[] =
 };
 
 
-/** Hull shader system values. */
-static FSystemValue HullSystemValueTable[] =
+/** Unsupported shader system values: Hull and Domain no longer supported in UE5, Mesh and Amplification shaders never supported for GLSL backend. */
+static FSystemValue DummySystemValueTable[] =
 {
-	{ "SV_OutputControlPointID", glsl_type::int_type, "gl_InvocationID", ir_var_in, false, false, false, false },
-	{ NULL, NULL, NULL, ir_var_auto, false, false, false, false }
-};
-
-/** Domain shader system values. */
-static FSystemValue DomainSystemValueTable[] =
-{
-
-	// TODO : SV_DomainLocation has types float2 or float3 depending on the input topology
-	{ "SV_Position", glsl_type::vec4_type, "gl_Position", ir_var_in, false, true, true, false },
-	{ "SV_Position", glsl_type::vec4_type, "gl_Position", ir_var_out, false, false, true, false },
-	{ "SV_DomainLocation", glsl_type::vec3_type, "gl_TessCoord", ir_var_in, false, false, false, false },
 	{ NULL, NULL, NULL, ir_var_auto, false, false, false, false }
 };
 
@@ -4038,8 +4026,8 @@ FSystemValue* SystemValueTable[HSF_FrequencyCount] =
 	VertexSystemValueTable,
 	PixelSystemValueTable,
 	GeometrySystemValueTable,
-	HullSystemValueTable,
-	DomainSystemValueTable,
+	DummySystemValueTable,
+	DummySystemValueTable,
 	ComputeSystemValueTable
 };
 
@@ -4276,46 +4264,6 @@ static ir_rvalue* GenShaderInputSemantic(
 		}
 	}
 
-	if (Variable == NULL && Frequency == HSF_DomainShader)
-	{
-		const int PrefixLength = 13;
-		if (FCStringAnsi::Strnicmp(Semantic, "SV_TessFactor", PrefixLength) == 0
-			&& Semantic[PrefixLength] >= '0'
-			&& Semantic[PrefixLength] <= '3')
-		{
-			int OutputIndex = Semantic[PrefixLength] - '0';
-			Variable = new(ParseState)ir_variable(
-				Type,
-				ralloc_asprintf(ParseState, "gl_TessLevelOuter[%d]", OutputIndex),
-				ir_var_out
-				);
-		}
-	}
-
-	if (Variable == NULL && Frequency == HSF_DomainShader)
-	{
-		const int PrefixLength = 19;
-		if (FCStringAnsi::Strnicmp(Semantic, "SV_InsideTessFactor", PrefixLength) == 0
-			&& Semantic[PrefixLength] >= '0'
-			&& Semantic[PrefixLength] <= '1')
-		{
-			int OutputIndex = Semantic[PrefixLength] - '0';
-			Variable = new(ParseState)ir_variable(
-				Type,
-				ralloc_asprintf(ParseState, "gl_TessLevelInner[%d]", OutputIndex),
-				ir_var_out
-				);
-		}
-		else if (FCStringAnsi::Stricmp(Semantic, "SV_InsideTessFactor") == 0)
-		{
-			Variable = new(ParseState)ir_variable(
-				Type,
-				ralloc_asprintf(ParseState, "gl_TessLevelInner[0]"),
-				ir_var_out
-				);
-		}
-	}
-
 	if (Variable)
 	{
 		// Up to this point, variables aren't contained in structs
@@ -4511,7 +4459,7 @@ static ir_rvalue* GenShaderOutputSemantic(
 		}
 	}
 
-	if (Variable == NULL && (Frequency == HSF_VertexShader || Frequency == HSF_GeometryShader || Frequency == HSF_HullShader || Frequency == HSF_DomainShader))
+	if (Variable == NULL && (Frequency == HSF_VertexShader || Frequency == HSF_GeometryShader))
 	{
 		const int PrefixLength = 15;
 		if (FCStringAnsi::Strnicmp(Semantic, "SV_ClipDistance", PrefixLength) == 0)
@@ -4558,52 +4506,6 @@ static ir_rvalue* GenShaderOutputSemantic(
 		}
 	}
 
-	if (Variable == NULL && Frequency == HSF_HullShader)
-	{
-		const int PrefixLength = 13;
-		if (FCStringAnsi::Strnicmp(Semantic, "SV_TessFactor", PrefixLength) == 0
-			&& Semantic[PrefixLength] >= '0'
-			&& Semantic[PrefixLength] <= '3')
-		{
-			int OutputIndex = Semantic[PrefixLength] - '0';
-			Variable = new(ParseState)ir_variable(
-				Type,
-				ralloc_asprintf(ParseState, "gl_TessLevelOuter[%d]", OutputIndex),
-				ir_var_out
-				);
-
-			ApplyClampPowerOfTwo = ParseState->tessellation.partitioning == GLSL_PARTITIONING_POW2;
-		}
-	}
-
-	if (Variable == NULL && Frequency == HSF_HullShader)
-	{
-		const int PrefixLength = 19;
-		if (FCStringAnsi::Strnicmp(Semantic, "SV_InsideTessFactor", PrefixLength) == 0
-			&& Semantic[PrefixLength] >= '0'
-			&& Semantic[PrefixLength] <= '1')
-		{
-			int OutputIndex = Semantic[PrefixLength] - '0';
-			Variable = new(ParseState)ir_variable(
-				Type,
-				ralloc_asprintf(ParseState, "gl_TessLevelInner[%d]", OutputIndex),
-				ir_var_out
-				);
-
-			ApplyClampPowerOfTwo = ParseState->tessellation.partitioning == GLSL_PARTITIONING_POW2;
-		}
-		else if (FCStringAnsi::Stricmp(Semantic, "SV_InsideTessFactor") == 0)
-		{
-			Variable = new(ParseState)ir_variable(
-				Type,
-				ralloc_asprintf(ParseState, "gl_TessLevelInner[0]"),
-				ir_var_out
-				);
-
-			ApplyClampPowerOfTwo = ParseState->tessellation.partitioning == GLSL_PARTITIONING_POW2;
-		}
-	}
-
 	bool bUseGLSL410Rules = OutputQualifier.Fields.bIsPatchConstant && ParseState->language_version == 410;
 	bool bUseESRules = ParseState->bGenerateES || ParseState->language_version == 310;
 
@@ -4648,10 +4550,6 @@ static ir_rvalue* GenShaderOutputSemantic(
 
 	const glsl_type* VariableType = glsl_type::get_record_instance(StructField, 1, ralloc_strdup(ParseState, Semantic));
 
-	if (Frequency == HSF_HullShader && !OutputQualifier.Fields.bIsPatchConstant)
-	{
-		VariableType = glsl_type::get_array_instance(VariableType, ParseState->tessellation.outputcontrolpoints);
-	}
 	Variable = new(ParseState)ir_variable(VariableType, ralloc_asprintf(ParseState, "out_%s", Semantic), ir_var_out);
 
 	Variable->centroid = OutputQualifier.Fields.bCentroid;
@@ -4669,11 +4567,6 @@ static ir_rvalue* GenShaderOutputSemantic(
 
 	ir_rvalue* VariableDeref = new(ParseState)ir_dereference_variable(Variable);
 
-	if (Frequency == HSF_HullShader && !OutputQualifier.Fields.bIsPatchConstant)
-	{
-		VariableDeref = new(ParseState)ir_dereference_array(VariableDeref, new(ParseState)ir_dereference_variable(ParseState->symbols->get_variable("gl_InvocationID")));
-	}
-
 	VariableDeref = new(ParseState)ir_dereference_record(VariableDeref, ralloc_strdup(ParseState, "Data"));
 
 	return VariableDeref;
@@ -4914,15 +4807,7 @@ static ir_dereference_variable* GenShaderInput(
 
 
 	// everything that's not an Outputpatch is patch constant. System values are treated specially
-	if (Frequency == HSF_DomainShader && !InputType->is_outputpatch())
-	{
-		InputQualifier.Fields.bIsPatchConstant = true;
-	}
-
-	if ((Frequency == HSF_GeometryShader && TempVariableDeref->type->is_array()) ||
-		(Frequency == HSF_HullShader && TempVariableDeref->type->is_inputpatch()) ||
-		(Frequency == HSF_DomainShader && TempVariableDeref->type->is_outputpatch())
-		)
+	if (Frequency == HSF_GeometryShader && TempVariableDeref->type->is_array())
 	{
 		check(InputType->is_array() || InputType->is_inputpatch() || InputType->is_outputpatch());
 		check(InputType->length || InputType->patch_length);
@@ -5368,8 +5253,6 @@ bool FGlslCodeBackend::GenerateMain(
 			ir_variable *const Variable = (ir_variable *)Iter.get();
 			if (Variable->semantic != NULL || Variable->type->is_record()
 				|| (Frequency == HSF_GeometryShader && (Variable->type->is_outputstream() || Variable->type->is_array()))
-				|| (Frequency == HSF_HullShader && (Variable->type->is_patch()))
-				|| (Frequency == HSF_DomainShader && (Variable->type->is_outputpatch()))
 				)
 			{
 				FSemanticQualifier Qualifier;
@@ -5461,36 +5344,37 @@ bool FGlslCodeBackend::GenerateMain(
 					break;
 				case ir_var_inout:
 				{
-									 check(Frequency == HSF_GeometryShader);
-									 // This is an output stream for geometry shader. It's not referenced as a variable inside the function,
-									 // instead OutputStream.Append(vertex) and OutputStream.RestartStrip() are called, and this variable
-									 // has already been optimized out of them in ast_to_hir translation.
-
-									 // Generate a local variable to add to arguments. It won't be referenced anywhere, so it should get optimized out.
-									 ir_variable* TempVariable = new(ParseState)ir_variable(
-										 Variable->type,
-										 NULL,
-										 ir_var_temporary);
-									 ArgVarDeref = new(ParseState)ir_dereference_variable(TempVariable);
-									 PreCallInstructions.push_tail(TempVariable);
-
-									 // We need to move this information somewhere safer, as this pseudo-variable will get optimized out of existence
-									 ParseState->outputstream_type = Variable->type->outputstream_type;
-
-									 check(Variable->type->is_outputstream());
-									 check(Variable->type->inner_type->is_record());
-
-									 geometry_append_type = Variable->type->inner_type;
+					check(Frequency == HSF_GeometryShader);
+					// This is an output stream for geometry shader. It's not referenced as a variable inside the function,
+					// instead OutputStream.Append(vertex) and OutputStream.RestartStrip() are called, and this variable
+					// has already been optimized out of them in ast_to_hir translation.
+
+					// Generate a local variable to add to arguments. It won't be referenced anywhere, so it should get optimized out.
+					ir_variable* TempVariable = new(ParseState)ir_variable(
+						Variable->type,
+						NULL,
+						ir_var_temporary);
+					ArgVarDeref = new(ParseState)ir_dereference_variable(TempVariable);
+					PreCallInstructions.push_tail(TempVariable);
+
+					// We need to move this information somewhere safer, as this pseudo-variable will get optimized out of existence
+					ParseState->outputstream_type = Variable->type->outputstream_type;
+
+					check(Variable->type->is_outputstream());
+					check(Variable->type->inner_type->is_record());
+
+					geometry_append_type = Variable->type->inner_type;
 				}
 					break;
 				default:
 				{
-						   _mesa_glsl_error(
-							   ParseState,
-							   "entry point parameter '%s' must be an input or output",
-							   Variable->name
-							   );
+					_mesa_glsl_error(
+						ParseState,
+						"entry point parameter '%s' must be an input or output",
+						Variable->name
+						);
 				}
+					break;
 				}
 				ArgInstructions.push_tail(ArgVarDeref);
 			}
@@ -5527,65 +5411,6 @@ bool FGlslCodeBackend::GenerateMain(
 				);
 		}
 
-		/*
-		we map the HLSL hull shader to this GLSL main function
-		for the most parts, we treat variables of InputPatch and OutputPatch as arrays of the inner type
-
-		build input patch from shader input interface blocks
-		call hull shader main function with input patch and current control point id (gl_InvocationID)
-		copy hull shader main result for the current control point to the proper shader output interface block element
-		barrier
-		(so all instances have computed the per control point data)
-		build patch constant function input (of type output patch) from the shader output interface blocks
-		(need to do this, since this is the only shader variable shared between control points running in parallel)
-
-		if control point id (gl_InvocationID) is 0
-		call patch constant function with the output patch as an input
-		copy the patch constant result to the "patch" shader output interface block
-		*/
-
-		if (Frequency == HSF_HullShader)
-		{
-
-			ir_function_signature* PatchConstantSig = FindPatchConstantFunction(Instructions, ParseState);
-
-			if (!PatchConstantSig)
-			{
-				_mesa_glsl_error(ParseState, "patch constant function `%s' not found", ParseState->tessellation.patchconstantfunc);
-			}
-
-			const glsl_type* OutputPatchType = glsl_type::get_templated_instance(EntryPointReturn->type, "OutputPatch", 0, ParseState->tessellation.outputcontrolpoints);
-
-			ir_variable* OutputPatchVar = new(ParseState)ir_variable(OutputPatchType, NULL, ir_var_temporary);
-
-
-			// call barrier() to ensure that all threads have computed the per-patch computation
-			{
-				// We can't just use the symbol table b/c it only has the HLSL and not the GLSL barrier functions in it
-				foreach_iter(exec_list_iterator, Iter, *Instructions)
-				{
-					ir_instruction *ir = (ir_instruction *)Iter.get();
-					ir_function *Function = ir->as_function();
-					if (Function && strcmp(Function->name, "barrier") == 0)
-					{
-						check(Function->signatures.get_head() == Function->signatures.get_tail());
-						exec_list VoidParameter;
-						ir_function_signature * BarrierFunctionSig = Function->matching_signature(&VoidParameter);
-						PostCallInstructions.push_tail(new(ParseState)ir_call(BarrierFunctionSig, NULL, &VoidParameter));
-					}
-				}
-			}
-
-			// reassemble output patch variable(for the patch constant function) from the shader outputs
-			GenShaderPatchConstantFunctionInputs(ParseState, OutputPatchVar, PostCallInstructions);
-
-			// call the entry point
-			if (PatchConstantSig)
-			{
-				CallPatchConstantFunction(ParseState, OutputPatchVar, PatchConstantSig, DeclInstructions, PostCallInstructions);
-			}
-		}
-
 		ParseState->symbols->pop_scope();
 
 		// Build the void main() function for GLSL.
@@ -5626,32 +5451,32 @@ bool FGlslCodeBackend::GenerateMain(
 		// and not the hull shader, so we specify them for both in the .usf shaders and then print a warning,
 		// similar to what fxc is doing
 
-		if (MainSig->tessellation.domain != GLSL_DOMAIN_NONE && (Frequency != HSF_HullShader && Frequency != HSF_DomainShader))
+		if (MainSig->tessellation.domain != GLSL_DOMAIN_NONE)
 		{
 			_mesa_glsl_warning(ParseState, "'domain' attribute only applies to hull or domain shaders");
 		}
 
-		if (MainSig->tessellation.outputtopology != GLSL_OUTPUTTOPOLOGY_NONE && Frequency != HSF_HullShader)
+		if (MainSig->tessellation.outputtopology != GLSL_OUTPUTTOPOLOGY_NONE)
 		{
 			_mesa_glsl_warning(ParseState, "'outputtopology' attribute only applies to hull shaders");
 		}
 
-		if (MainSig->tessellation.partitioning != GLSL_PARTITIONING_NONE && Frequency != HSF_HullShader)
+		if (MainSig->tessellation.partitioning != GLSL_PARTITIONING_NONE)
 		{
 			_mesa_glsl_warning(ParseState, "'partitioning' attribute only applies to hull shaders");
 		}
 
-		if (MainSig->tessellation.outputcontrolpoints > 0 && Frequency != HSF_HullShader)
+		if (MainSig->tessellation.outputcontrolpoints > 0)
 		{
 			_mesa_glsl_warning(ParseState, "'outputcontrolpoints' attribute only applies to hull shaders");
 		}
 
-		if (MainSig->tessellation.maxtessfactor > 0.0f && Frequency != HSF_HullShader)
+		if (MainSig->tessellation.maxtessfactor > 0.0f)
 		{
 			_mesa_glsl_warning(ParseState, "'maxtessfactor' attribute only applies to hull shaders");
 		}
 
-		if (MainSig->tessellation.patchconstantfunc != 0 && Frequency != HSF_HullShader)
+		if (MainSig->tessellation.patchconstantfunc != 0)
 		{
 			_mesa_glsl_warning(ParseState, "'patchconstantfunc' attribute only applies to hull shaders");
 		}
@@ -5709,182 +5534,6 @@ ir_function_signature*  FGlslCodeBackend::FindPatchConstantFunction(exec_list* I
 	return PatchConstantSig;
 }
 
-void FGlslCodeBackend::CallPatchConstantFunction(_mesa_glsl_parse_state* ParseState, ir_variable* OutputPatchVar, ir_function_signature* PatchConstantSig, exec_list& DeclInstructions, exec_list &PostCallInstructions)
-{
-	exec_list PatchConstantArgs;
-	if (OutputPatchVar && !PatchConstantSig->parameters.is_empty())
-	{
-		PatchConstantArgs.push_tail(new(ParseState)ir_dereference_variable(OutputPatchVar));
-	}
-
-	ir_if* thread_if = new(ParseState)ir_if(
-		new(ParseState)ir_expression(
-		ir_binop_equal,
-		new (ParseState)ir_constant(
-		0
-		),
-		new (ParseState)ir_dereference_variable(
-		ParseState->symbols->get_variable("gl_InvocationID")
-		)
-		)
-		);
-
-	exec_list PrePatchConstCallInstructions;
-	exec_list PostPatchConstCallInstructions;
-
-	FSemanticQualifier Qualifier;
-	Qualifier.Fields.bIsPatchConstant = 1;
-
-	ir_dereference_variable* PatchConstantReturn = GenShaderOutput(
-		HSF_HullShader,
-		ParseState,
-		PatchConstantSig->return_semantic,
-		Qualifier,
-		PatchConstantSig->return_type,
-		&DeclInstructions,
-		&PrePatchConstCallInstructions,
-		&PostPatchConstCallInstructions
-		);
-
-	thread_if->then_instructions.append_list(&PrePatchConstCallInstructions);
-	thread_if->then_instructions.push_tail(new(ParseState)ir_call(PatchConstantSig, PatchConstantReturn, &PatchConstantArgs));
-	thread_if->then_instructions.append_list(&PostPatchConstCallInstructions);
-
-	PostCallInstructions.push_tail(thread_if);
-			}
-
-
-			/*
-reassemble output patch variable (for the patch constant function) from the shader outputs
-			turn this: (from the GenOutputs of calling the entry point main)
-
-			out_InnerMember[gl_InvocationID].Data = t2.Middle.Inner.Value;
-
-			into this:
-
-			//output_patch<FPNTessellationHSToDS> FPNTessellationHSToDS t3[3]; //output_patch<FPNTessellationHSToDS> ;
-			t3[0].Middle.Inner.Value = out_InnerMember[0].Data;
-			t3[1].Middle.Inner.Value = out_InnerMember[1].Data;
-			t3[2].Middle.Inner.Value = out_InnerMember[2].Data;
-
-			*/
-
-void FGlslCodeBackend::GenShaderPatchConstantFunctionInputs(_mesa_glsl_parse_state* ParseState, ir_variable* OutputPatchVar, exec_list &PostCallInstructions)
-{
-	PostCallInstructions.push_tail(OutputPatchVar);
-	foreach_iter(exec_list_iterator, Iter, PostCallInstructions)
-	{
-		ir_instruction *ir = (ir_instruction *)Iter.get();
-
-		ir_assignment* assignment = ir->as_assignment();
-
-		if (!assignment)
-		{
-			continue;
-		}
-
-		ir_dereference_record* lhs = assignment->lhs->as_dereference_record();
-		ir_rvalue* rhs = assignment->rhs;
-
-		if (!lhs)
-		{
-			continue;
-		}
-
-		if (!rhs)
-		{
-			continue;
-		}
-
-		ir_dereference_array* lhs_array = lhs->record->as_dereference_array();
-
-		if (!lhs_array)
-		{
-			continue;
-		}
-
-		ir_dereference_variable* OutputPatchArrayIndex = lhs_array->array_index->as_dereference_variable();
-		ir_dereference_variable* OutputPatchArray = lhs_array->array->as_dereference_variable();
-
-		if (!OutputPatchArrayIndex)
-		{
-			continue;
-		}
-
-		if (0 != strcmp(OutputPatchArrayIndex->var->name, "gl_InvocationID"))
-		{
-			continue;
-		}
-
-		if (!OutputPatchArray)
-		{
-			continue;
-		}
-
-		const char* OutArrayFieldName = lhs->field;
-
-		for (int OutputVertex = 0; OutputVertex < ParseState->tessellation.outputcontrolpoints; ++OutputVertex)
-		{
-			struct Helper
-			{
-				// the struct inside the output patch can have the actual outputs with semantics nested inside,
-
-				static void ReplaceVariableDerefWithArrayDeref(ir_instruction* Node, ir_dereference_array* ArrayDereference)
-				{
-					if (ir_dereference_record* AsRecord = Node->as_dereference_record())
-					{
-						if (AsRecord->record->as_dereference_variable())
-						{
-							AsRecord->record = ArrayDereference;
-						}
-						else
-						{
-							ReplaceVariableDerefWithArrayDeref(AsRecord->record, ArrayDereference);
-						}
-					}
-					else if (ir_dereference_array* AsArray = Node->as_dereference_array())
-					{
-						if (AsArray->array->as_dereference_variable())
-						{
-							AsArray->array = ArrayDereference;
-						}
-						else
-						{
-							ReplaceVariableDerefWithArrayDeref(AsArray->array, ArrayDereference);
-						}
-					}
-					else
-					{
-						check(false);
-					}
-				}
-			};
-
-			ir_dereference_array* OutputPatchElementIndex = new(ParseState)ir_dereference_array(
-				OutputPatchVar,
-				new(ParseState)ir_constant(
-				OutputVertex
-				)
-				);
-
-			ir_rvalue* OutputPatchElement = rhs->clone(ParseState, 0);
-			Helper::ReplaceVariableDerefWithArrayDeref(OutputPatchElement, OutputPatchElementIndex);
-
-			PostCallInstructions.push_tail(
-				new (ParseState)ir_assignment(
-					OutputPatchElement,
-					new(ParseState)ir_dereference_record(
-						new(ParseState)ir_dereference_array(
-							OutputPatchArray->clone(ParseState, 0),
-							new(ParseState)ir_constant(OutputVertex)
-						),
-						OutArrayFieldName
-					)
-				)
-				);
-		}
-	}
-}
 
 void FGlslLanguageSpec::SetupLanguageIntrinsics(_mesa_glsl_parse_state* State, exec_list* ir)
 {
diff --git a/Engine/Source/Developer/ShaderFormatOpenGL/Private/OpenGLShaderCompiler.cpp b/Engine/Source/Developer/ShaderFormatOpenGL/Private/OpenGLShaderCompiler.cpp
--- a/Engine/Source/Developer/ShaderFormatOpenGL/Private/OpenGLShaderCompiler.cpp
+++ b/Engine/Source/Developer/ShaderFormatOpenGL/Private/OpenGLShaderCompiler.cpp
@@ -1211,31 +1211,15 @@ FGlslLanguageSpec* FOpenGLFrontend::CreateLanguageSpec(GLSLVersion Version, bool
 
 #if DXC_SUPPORTED
 
-static const ANSICHAR* GetFrequencyPrefix(EHlslShaderFrequency Frequency)
+static const ANSICHAR* GetFrequencyPrefix(EShaderFrequency Frequency)
 {
 	switch (Frequency)
 	{
-	case HSF_VertexShader:		return "v";
-	case HSF_PixelShader:		return "p";
-	case HSF_GeometryShader:	return "g";
-	case HSF_HullShader:		return "h";
-	case HSF_DomainShader:		return "d";
-	case HSF_ComputeShader:		return "c";
-	default:					return "";
-	}
-}
-
-static const TCHAR* GetFrequencyFileExt(EHlslShaderFrequency Frequency)
-{
-	switch (Frequency)
-	{
-	case HSF_VertexShader:		return TEXT("vert");
-	case HSF_PixelShader:		return TEXT("frag");
-	case HSF_GeometryShader:	return TEXT("geom");
-	case HSF_HullShader:		return TEXT("tesc");
-	case HSF_DomainShader:		return TEXT("tese");
-	case HSF_ComputeShader:		return TEXT("comp");
-	default:					return TEXT("glsl");
+	case SF_Vertex:		return "v";
+	case SF_Pixel:		return "p";
+	case SF_Geometry:	return "g";
+	case SF_Compute:	return "c";
+	default:			return "";
 	}
 }
 
@@ -1463,7 +1447,7 @@ static bool CompileToGlslWithShaderConductor(
 	FShaderCompilerOutput&		Output,
 	const FString&				WorkingDirectory,
 	GLSLVersion					Version,
-	const EHlslShaderFrequency	Frequency,
+	const EShaderFrequency		Frequency,
 	uint32						CCFlags,
 	const FString&				PreprocessedShader,
 	char*&						OutGlslShaderSource)
@@ -2189,7 +2173,7 @@ static bool CompileToGlslWithShaderConductor(
 
 					GlslSource.insert(LayoutPos, DefineString);
 				}
-				if(!bEmulatedUBs)
+				/*if(!bEmulatedUBs)
 				{
 					for (auto const& Pair : UniformVarNames)
 					{
@@ -2201,7 +2185,7 @@ static bool CompileToGlslWithShaderConductor(
 
 						GlslSource.insert(LayoutPos, DefineString);
 					}
-				}
+				}*/
 			}
 
 
@@ -2452,14 +2436,14 @@ static bool CompileToGlslWithShaderConductor(
 				}
 			}
 
-			size_t GlslSourceLen = GlslSource.length();
+			const size_t GlslSourceLen = GlslSource.length();
 			if (GlslSourceLen > 0)
 			{
-				size_t SamplerPos = GlslSource.find("\nuniform ");
-
 				uint32 TextureIndex = 0;
 				for (const FString& Texture : Textures)
 				{
+					const size_t SamplerPos = GlslSource.find("\nuniform ");
+
 					TArray<FString> UsedSamplers;
 					FString SamplerString;
 					for (const FString& Sampler : Samplers)
@@ -2469,6 +2453,8 @@ static bool CompileToGlslWithShaderConductor(
 						size_t FindCombinedSampler = GlslSource.find(SamplerName.c_str());
 						if (FindCombinedSampler != std::string::npos)
 						{
+							checkf(SamplerPos != std::string::npos, TEXT("Generated GLSL shader is expected to have combined sampler '%s:%s' but no appropriate 'uniform' declaration could be found"), *Texture, *Sampler);
+
 							uint32 NewIndex = TextureIndex + UsedSamplers.Num();
 							std::string NewDefine = "#define ";
 							NewDefine += SamplerName;
@@ -2543,7 +2529,8 @@ static bool CompileToGlslWithShaderConductor(
 
 	if (bDumpDebugInfo && OutGlslShaderSource != nullptr)
 	{
-		DumpDebugShaderText(Input, OutGlslShaderSource, FCStringAnsi::Strlen(OutGlslShaderSource), GetFrequencyFileExt(Frequency));
+		const TCHAR* ShaderFileExt = CrossCompiler::FShaderConductorContext::GetShaderFileExt(CrossCompiler::EShaderConductorLanguage::Essl, Frequency);
+		DumpDebugShaderText(Input, OutGlslShaderSource, FCStringAnsi::Strlen(OutGlslShaderSource), ShaderFileExt);
 	}
 
 	return !bCompilationFailed;
@@ -2649,14 +2636,16 @@ void FOpenGLFrontend::CompileShader(const FShaderCompilerInput& Input, FShaderCo
 	};
 	static_assert(SF_NumStandardFrequencies == UE_ARRAY_COUNT(FrequencyTable), "NumFrequencies changed. Please update tables.");
 
-	const EHlslShaderFrequency Frequency = FrequencyTable[Input.Target.Frequency];
-	if (Frequency == HSF_InvalidFrequency)
+	const EShaderFrequency Frequency = (EShaderFrequency)Input.Target.Frequency;
+
+	const EHlslShaderFrequency HlslFrequency = FrequencyTable[Frequency];
+	if (HlslFrequency == HSF_InvalidFrequency)
 	{
 		Output.bSucceeded = false;
 		FShaderCompilerError* NewError = new(Output.Errors) FShaderCompilerError();
 		NewError->StrippedErrorMessage = FString::Printf(
 			TEXT("%s shaders not supported for use in OpenGL."),
-			CrossCompiler::GetFrequencyName((EShaderFrequency)Input.Target.Frequency)
+			CrossCompiler::GetFrequencyName(Frequency)
 			);
 		return;
 	}
@@ -2703,7 +2692,7 @@ void FOpenGLFrontend::CompileShader(const FShaderCompilerInput& Input, FShaderCo
 		bool bDefaultPrecisionIsHalf = (CCFlags & HLSLCC_UseFullPrecisionInPS) == 0;
 		FGlslLanguageSpec* LanguageSpec = CreateLanguageSpec(Version, bDefaultPrecisionIsHalf);
 
-		FHlslCrossCompilerContext CrossCompilerContext(CCFlags, Frequency, HlslCompilerTarget);
+		FHlslCrossCompilerContext CrossCompilerContext(CCFlags, HlslFrequency, HlslCompilerTarget);
 		if (CrossCompilerContext.Init(TCHAR_TO_ANSI(*Input.VirtualSourceFilePath), LanguageSpec))
 		{
 			bCompilationSucceeded = CrossCompilerContext.Run(
@@ -2720,7 +2709,8 @@ void FOpenGLFrontend::CompileShader(const FShaderCompilerInput& Input, FShaderCo
 
 		if (bDumpDebugInfo && bCompilationSucceeded && GlslShaderSource != nullptr)
 		{
-			DumpDebugShaderText(Input, GlslShaderSource, FCStringAnsi::Strlen(GlslShaderSource), GetFrequencyFileExt(Frequency));
+			const TCHAR* ShaderFileExt = CrossCompiler::FShaderConductorContext::GetShaderFileExt(CrossCompiler::EShaderConductorLanguage::Essl, Frequency);
+			DumpDebugShaderText(Input, GlslShaderSource, FCStringAnsi::Strlen(GlslShaderSource), ShaderFileExt);
 		}
 	}
 
@@ -2733,7 +2723,7 @@ void FOpenGLFrontend::CompileShader(const FShaderCompilerInput& Input, FShaderCo
 		}
 
 #if VALIDATE_GLSL_WITH_DRIVER
-		PrecompileShader(Output, Input, GlslShaderSource, Version, Frequency);
+		PrecompileShader(Output, Input, GlslShaderSource, Version, HlslFrequency);
 #else // VALIDATE_GLSL_WITH_DRIVER
 		int32 SourceLen = FCStringAnsi::Strlen(GlslShaderSource); //-V595
 		Output.Target = Input.Target;
@@ -2742,7 +2732,8 @@ void FOpenGLFrontend::CompileShader(const FShaderCompilerInput& Input, FShaderCo
 
 		if (bDumpDebugInfo)
 		{
-			FString DumpedGlslFile = *Input.DumpDebugInfoPath / (TEXT("Output") + GetExtension(Frequency));
+			const TCHAR* ShaderFileExt = CrossCompiler::FShaderConductorContext::GetShaderFileExt(CrossCompiler::EShaderConductorLanguage::Essl, Frequency);
+			FString DumpedGlslFile = FString::Printf(TEXT("%s/Output.%s"), *Input.DumpDebugInfoPath, ShaderFileExt);
 			if (TUniquePtr<FArchive> FileWriter = TUniquePtr<FArchive>(IFileManager::Get().CreateFileWriter(*DumpedGlslFile)))
 			{
 				FileWriter->Serialize(GlslShaderSource, FCStringAnsi::Strlen(GlslShaderSource));
diff --git a/Engine/Source/Developer/ShaderFormatOpenGL/Public/GlslBackend.h b/Engine/Source/Developer/ShaderFormatOpenGL/Public/GlslBackend.h
--- a/Engine/Source/Developer/ShaderFormatOpenGL/Public/GlslBackend.h
+++ b/Engine/Source/Developer/ShaderFormatOpenGL/Public/GlslBackend.h
@@ -72,10 +72,6 @@ struct SHADERFORMATOPENGL_API FGlslCodeBackend : public FCodeBackend
 	*/
 	virtual bool GenerateMain(EHlslShaderFrequency Frequency, const char* EntryPoint, exec_list* Instructions, _mesa_glsl_parse_state* ParseState) override;
 
-	void GenShaderPatchConstantFunctionInputs(_mesa_glsl_parse_state* ParseState, ir_variable* OutputPatchVar, exec_list &PostCallInstructions);
-
-	void CallPatchConstantFunction(_mesa_glsl_parse_state* ParseState, ir_variable* OutputPatchVar, ir_function_signature* PatchConstantSig, exec_list& DeclInstructions, exec_list &PostCallInstructions);
-
 	ir_function_signature* FindPatchConstantFunction(exec_list* Instructions, _mesa_glsl_parse_state* ParseState);
 
 
diff --git a/Engine/Source/Developer/VulkanShaderFormat/Private/VulkanShaderCompiler.cpp b/Engine/Source/Developer/VulkanShaderFormat/Private/VulkanShaderCompiler.cpp
--- a/Engine/Source/Developer/VulkanShaderFormat/Private/VulkanShaderCompiler.cpp
+++ b/Engine/Source/Developer/VulkanShaderFormat/Private/VulkanShaderCompiler.cpp
@@ -1947,6 +1947,7 @@ static bool BuildShaderOutputFromSpirv(
 static bool CompileWithShaderConductor(
 	const FString&			PreprocessedShader,
 	const FString&			EntryPointName,
+	EShaderFrequency		Frequency,
 	const FCompilerInfo&	CompilerInfo,
 	EHlslCompileTarget		HlslCompilerTarget,
 	FShaderCompilerOutput&	Output,
@@ -1971,7 +1972,7 @@ static bool CompileWithShaderConductor(
 	}
 
 	// Load shader source into compiler context
-	CompilerContext.LoadSource(PreprocessedShader, Input.VirtualSourceFilePath, EntryPointName, CompilerInfo.Frequency, &AdditionalDefines);
+	CompilerContext.LoadSource(PreprocessedShader, Input.VirtualSourceFilePath, EntryPointName, Frequency, &AdditionalDefines);
 
 	// Initialize compilation options for ShaderConductor
 	CrossCompiler::FShaderConductorOptions Options;
@@ -2048,14 +2049,16 @@ void DoCompileVulkanShader(const FShaderCompilerInput& Input, FShaderCompilerOut
 		bIsSM5 ? HSF_RayCallable : HSF_InvalidFrequency,
 	};
 
-	const EHlslShaderFrequency Frequency = FrequencyTable[Input.Target.Frequency];
-	if (Frequency == HSF_InvalidFrequency)
+	const EShaderFrequency Frequency = (EShaderFrequency)Input.Target.Frequency;
+
+	const EHlslShaderFrequency HlslFrequency = FrequencyTable[Input.Target.Frequency];
+	if (HlslFrequency == HSF_InvalidFrequency)
 	{
 		Output.bSucceeded = false;
 		FShaderCompilerError* NewError = new(Output.Errors) FShaderCompilerError();
 		NewError->StrippedErrorMessage = FString::Printf(
 			TEXT("%s shaders not supported for use in Vulkan."),
-			CrossCompiler::GetFrequencyName((EShaderFrequency)Input.Target.Frequency));
+			CrossCompiler::GetFrequencyName(Frequency));
 		return;
 	}
 
@@ -2129,7 +2132,7 @@ void DoCompileVulkanShader(const FShaderCompilerInput& Input, FShaderCompilerOut
 
 	RemoveUniformBuffersFromSource(Input.Environment, PreprocessedShaderSource);
 
-	FCompilerInfo CompilerInfo(Input, WorkingDirectory, Frequency);
+	FCompilerInfo CompilerInfo(Input, WorkingDirectory, HlslFrequency);
 
 	// Setup hlslcc flags. Needed here as it will be used when dumping debug info
 	{
@@ -2174,7 +2177,7 @@ void DoCompileVulkanShader(const FShaderCompilerInput& Input, FShaderCompilerOut
 
 #if PLATFORM_MAC || PLATFORM_WINDOWS || PLATFORM_LINUX
 	// Cross-compile shader via ShaderConductor (DXC, SPIRV-Tools, SPIRV-Cross)
-	bSuccess = CompileWithShaderConductor(PreprocessedShaderSource, EntryPointName, CompilerInfo, HlslCompilerTarget, Output, BindingTable, bStripReflect);
+	bSuccess = CompileWithShaderConductor(PreprocessedShaderSource, EntryPointName, Frequency, CompilerInfo, HlslCompilerTarget, Output, BindingTable, bStripReflect);
 #endif // PLATFORM_MAC || PLATFORM_WINDOWS || PLATFORM_LINUX
 	
 	ShaderParameterParser.ValidateShaderParameterTypes(Input, Output);
diff --git a/Engine/Source/Developer/Windows/ShaderFormatD3D/Private/D3DShaderCompiler.cpp b/Engine/Source/Developer/Windows/ShaderFormatD3D/Private/D3DShaderCompiler.cpp
--- a/Engine/Source/Developer/Windows/ShaderFormatD3D/Private/D3DShaderCompiler.cpp
+++ b/Engine/Source/Developer/Windows/ShaderFormatD3D/Private/D3DShaderCompiler.cpp
@@ -643,7 +643,7 @@ bool CompileAndProcessD3DShaderFXC(FString& PreprocessedShaderSource, const FStr
 			CrossCompiler::FShaderConductorContext CompilerContext;
 
 			// Load shader source into compiler context
-			const EHlslShaderFrequency Frequency = CrossCompiler::FShaderConductorContext::MapShaderFrequency(static_cast<EShaderFrequency>(Input.Target.Frequency));
+			const EShaderFrequency Frequency = (EShaderFrequency)Input.Target.Frequency;
 			CompilerContext.LoadSource(PreprocessedShaderSource, Input.VirtualSourceFilePath, EntryPointName, Frequency);
 
 			// Compile HLSL source to SPIR-V binary
