Fix compiler errors in Metal and GLSL backends.

#rnx

diff --git a/Engine/Source/Developer/Apple/MetalShaderFormat/Private/MetalShaderCompiler.cpp b/Engine/Source/Developer/Apple/MetalShaderFormat/Private/MetalShaderCompiler.cpp
--- a/Engine/Source/Developer/Apple/MetalShaderFormat/Private/MetalShaderCompiler.cpp
+++ b/Engine/Source/Developer/Apple/MetalShaderFormat/Private/MetalShaderCompiler.cpp
@@ -243,7 +243,7 @@ void BuildMetalShaderOutput(
 		UE_LOG(LogMetalShaderCompiler, Fatal, TEXT("Bad hlslcc header found"));
 	}
 	
-	EShaderFrequency Frequency = (EShaderFrequency)ShaderOutput.Target.Frequency;
+	const EShaderFrequency Frequency = (EShaderFrequency)ShaderOutput.Target.Frequency;
 
 	//TODO read from toolchain
 	const bool bIsMobile = (ShaderInput.Target.Platform == SP_METAL || ShaderInput.Target.Platform == SP_METAL_MRT || ShaderInput.Target.Platform == SP_METAL_TVOS || ShaderInput.Target.Platform == SP_METAL_MRT_TVOS);
@@ -1132,7 +1132,7 @@ void CompileShader_Metal(const FShaderCompilerInput& _Input,FShaderCompilerOutpu
 	char* MetalShaderSource = NULL;
 	char* ErrorLog = NULL;
 
-	const EShaderFrequency Frequency = Input.Target.Frequency;
+	const EShaderFrequency Frequency = (EShaderFrequency)Input.Target.Frequency;
 	if (!(Frequency == SF_Vertex || Frequency == SF_Pixel || Frequency == SF_Compute))
 	{
 		Output.bSucceeded = false;
diff --git a/Engine/Source/Developer/ShaderFormatOpenGL/Private/OpenGLShaderCompiler.cpp b/Engine/Source/Developer/ShaderFormatOpenGL/Private/OpenGLShaderCompiler.cpp
--- a/Engine/Source/Developer/ShaderFormatOpenGL/Private/OpenGLShaderCompiler.cpp
+++ b/Engine/Source/Developer/ShaderFormatOpenGL/Private/OpenGLShaderCompiler.cpp
@@ -1865,7 +1865,7 @@ static bool CompileToGlslWithShaderConductor(
 			{
 				if (Var->storage_class == SpvStorageClassOutput && Var->built_in == -1 && !CrossCompiler::FShaderConductorContext::IsIntermediateSpirvOutputVariable(Var->name))
 				{
-					if (Frequency == HSF_PixelShader && strstr(Var->name, "SV_Target"))
+					if (Frequency == SF_Pixel && strstr(Var->name, "SV_Target"))
 					{
 						FString TypeQualifier;
 
@@ -2206,7 +2206,7 @@ static bool CompileToGlslWithShaderConductor(
 					std::string FBDFReplaceString = "_Globals_ARM_shader_framebuffer_fetch_depth_stencil";
 
 					size_t FramebufferDepthFetchPos = GlslSource.find(FBDFString);
-					bool UsesFramebufferDepthFetch = Frequency == HSF_PixelShader && GlslSource.find("_Globals.ARM_shader_framebuffer_fetch_depth_stencil") != std::string::npos;
+					bool UsesFramebufferDepthFetch = Frequency == SF_Pixel && GlslSource.find("_Globals.ARM_shader_framebuffer_fetch_depth_stencil") != std::string::npos;
 
 					if (UsesFramebufferDepthFetch)
 					{
@@ -2258,7 +2258,7 @@ static bool CompileToGlslWithShaderConductor(
 						{
 							GlslSource.erase(GlobalPos, GlobalEndPos - GlobalPos + GlobalsEndSearchString.length());
 
-							bool UsesFramebufferFetch = Frequency == HSF_PixelShader && GlslSource.find("_Globals.ARM_shader_framebuffer_fetch") != std::string::npos;
+							bool UsesFramebufferFetch = Frequency == SF_Pixel && GlslSource.find("_Globals.ARM_shader_framebuffer_fetch") != std::string::npos;
 
 							std::string GlobalVarString = "_Globals.";
 							size_t GlobalVarPos = 0;
