Fixed: Crash on running mobile preview with GL RHI

diff --git a/Engine/Source/Developer/ShaderFormatOpenGL/Private/OpenGLShaderCompiler.cpp b/Engine/Source/Developer/ShaderFormatOpenGL/Private/OpenGLShaderCompiler.cpp
--- a/Engine/Source/Developer/ShaderFormatOpenGL/Private/OpenGLShaderCompiler.cpp
+++ b/Engine/Source/Developer/ShaderFormatOpenGL/Private/OpenGLShaderCompiler.cpp
@@ -698,6 +698,8 @@ void FOpenGLFrontend::BuildShaderOutput(
 	for (auto& UniformBlock : CCHeader.UniformBlocks)
 	{
 		uint16 UBIndex = UniformBlock.Index;
+
+		UsedUniformBufferSlots[UBIndex] = true;
 		
 		if (OutputTrueParameterNames())
 		{
