OpenGL DXC Improvements
- Added support for structured buffers
- Updated OpenGLES to 3.2

diff --git a/Engine/Source/Developer/ShaderFormatOpenGL/Private/OpenGLShaderCompiler.cpp b/Engine/Source/Developer/ShaderFormatOpenGL/Private/OpenGLShaderCompiler.cpp
--- a/Engine/Source/Developer/ShaderFormatOpenGL/Private/OpenGLShaderCompiler.cpp
+++ b/Engine/Source/Developer/ShaderFormatOpenGL/Private/OpenGLShaderCompiler.cpp
@@ -1561,9 +1561,6 @@ static bool CompileToGlslWithShaderConductor(
 		"#else\n"
 		"#define _Globals_ARM_shader_framebuffer_fetch_depth_stencil 0u\n"
 		"#endif\n"
-		"#ifndef HLSLCC_DX11ClipSpace\n"
-		"#define HLSLCC_DX11ClipSpace 1\n"
-		"#endif\n"
 		;
 	
 	// Inject additional macro definitions to circumvent missing features: external textures
@@ -1640,6 +1637,7 @@ static bool CompileToGlslWithShaderConductor(
 		TArray<FString> Textures;
 		TArray<FString> Samplers;
 		TArray<std::string> UAVs;
+		TArray<std::string> StructuredBuffers;
 		uint32 UAVIndices = 0xffffffff;
 		uint32 BufferIndices = 0xffffffff;
 		uint32 TextureIndices = 0xffffffff;
@@ -1680,11 +1678,10 @@ static bool CompileToGlslWithShaderConductor(
 			UAVIndices &= ~(1 << Index);
 
 			CCHeaderWriter.WriteUAV(UTF8_TO_TCHAR(Binding->name), Index);
+			StructuredBuffers.Add(Binding->name);
 
 			SPVRResult = Reflection.ChangeDescriptorBindingNumbers(Binding, Index, GlobalSetId);
 			check(SPVRResult == SPV_REFLECT_RESULT_SUCCESS);
-
-			UAVs.Add(Binding->name);
 		}
 
 		for (auto const& Binding : ReflectionBindings.TextureUAVs)
@@ -1712,7 +1709,6 @@ static bool CompileToGlslWithShaderConductor(
 			uint32 Index = FPlatformMath::CountTrailingZeros(TextureIndices);
 
 			// No support for 3-component types in dxc/SPIRV/MSL - need to expose my workarounds there too
-			BufferIndices &= ~(1 << Index);
 			TextureIndices &= ~(1llu << uint64(Index));
 
 			Textures.Add(UTF8_TO_TCHAR(Binding->name));
@@ -1721,8 +1717,9 @@ static bool CompileToGlslWithShaderConductor(
 			check(SPVRResult == SPV_REFLECT_RESULT_SUCCESS);
 		}
 
-		for (auto const& Binding : ReflectionBindings.SBufferSRVs)
+		for (int32_t SBufferIndex = ReflectionBindings.SBufferSRVs.Num() - 1; SBufferIndex >= 0; SBufferIndex--)
 		{
+			auto Binding = ReflectionBindings.SBufferSRVs[SBufferIndex];
 			check(BufferIndices);
 			uint32 Index = FPlatformMath::CountTrailingZeros(BufferIndices);
 
@@ -1730,6 +1727,9 @@ static bool CompileToGlslWithShaderConductor(
 
 			SPVRResult = Reflection.ChangeDescriptorBindingNumbers(Binding, Index, GlobalSetId);
 			check(SPVRResult == SPV_REFLECT_RESULT_SUCCESS);
+
+			CCHeaderWriter.WriteUAV(UTF8_TO_TCHAR(Binding->name), Index);
+			StructuredBuffers.Add(Binding->name);
 		}
 
 		TMap<FString, uint32> GlobalOffsets;
@@ -2156,7 +2156,7 @@ static bool CompileToGlslWithShaderConductor(
 			TargetDesc.CompileFlags.SetDefine(TEXT("force_flattened_io_blocks"), 1);
 			TargetDesc.CompileFlags.SetDefine(TEXT("emit_uniform_buffer_as_plain_uniforms"), 1);
 			TargetDesc.Language = CrossCompiler::EShaderConductorLanguage::Essl;
-			TargetDesc.Version = 310;
+			TargetDesc.Version = 320;
 			break;
 		}
 
@@ -2263,10 +2263,17 @@ static bool CompileToGlslWithShaderConductor(
 			// Perform FBF replacements
 			{
 
-				size_t MainPos = GlslSource.find("#version 310 es");
+				size_t MainPos = GlslSource.find("#version 320 es");
+
+				// Fallback if the shader is 310 
+				if (MainPos == std::string::npos)
+				{
+					MainPos = GlslSource.find("#version 310 es");
+				}
+
 				if (MainPos != std::string::npos)
 				{
-					MainPos += strlen("#version 310 es");
+					MainPos += strlen("#version 320 es");
 					GlslSource.insert(MainPos, GlslFrameBufferDefines);
 					GlslSource.insert(MainPos, GlslFrameBufferExtensions);
 				}
@@ -2559,12 +2566,9 @@ static bool CompileToGlslWithShaderConductor(
 					while (std::regex_search(TempString, RegexMatch, std::regex(RegexExpression)))
 					{
 						RegexMatches.push_back(RegexMatch.str());
-						UE_LOG(LogOpenGLShaderCompiler, Warning, TEXT("\t\tFound Regex Result %s"), UTF8_TO_TCHAR(RegexMatch.str().c_str()));
 						TempString = RegexMatch.suffix().first;
 					}
 
-					UE_LOG(LogOpenGLShaderCompiler, Warning, TEXT("Looking for UAV %s"), UTF8_TO_TCHAR(UAV.c_str()));
-
 					for(const auto & Match : RegexMatches)
 					{
 						size_t MatchOffset = Match.find_first_of(UAV);
@@ -2576,6 +2580,44 @@ static bool CompileToGlslWithShaderConductor(
 					UAVIndex++;
 				}
 
+				// StructuredBuffers, rename as ci0 format
+				for (const std::string& SBuffer : StructuredBuffers)
+				{
+					std::string NewSBufferName = FrequencyPrefix + std::string("i") + std::to_string(UAVIndex);
+					std::string NewSBufferVarName = NewSBufferName + std::string("_VAR");
+					
+					std::string SearchString = "} " + SBuffer;
+					size_t SBufferPos = GlslSource.find(SearchString);
+
+					size_t SBufferEndPos = GlslSource.find(";", SBufferPos);
+					
+					std::string ReplacementSubStr = GlslSource.substr(SBufferPos + 2, SBufferEndPos - SBufferPos - 2);
+					UE_LOG(LogOpenGLShaderCompiler, Warning, TEXT("Found ReplacementSubStr %s"), UTF8_TO_TCHAR(ReplacementSubStr.c_str()));
+					GlslSource.erase(SBufferPos + 1, SBufferEndPos - SBufferPos - 1);
+
+					const std::string SBufferData = "_m0";
+					size_t SBufferDataPos = GlslSource.rfind(SBufferData, SBufferPos);
+					GlslSource.replace(SBufferDataPos, SBufferData.size(), NewSBufferName);
+
+					const std::string BufferString = " buffer ";
+					size_t BufferNamePos = GlslSource.rfind(BufferString, SBufferPos);
+					size_t BufferLineEndPos = GlslSource.find("\n", BufferNamePos);
+					size_t ReplacePos = BufferNamePos + BufferString.size();
+
+					GlslSource.replace(ReplacePos, BufferLineEndPos - ReplacePos, NewSBufferVarName);
+
+					// Replace the usage of StructuredBuffer with new name
+					size_t CurPos = GlslSource.find(ReplacementSubStr + ".");
+					while (CurPos != std::string::npos)
+					{
+						// Offset by 4 to account for ._m0
+						GlslSource.replace(CurPos, ReplacementSubStr.size()+4, NewSBufferName);
+						CurPos = GlslSource.find(ReplacementSubStr + ".");
+					}
+					
+					UAVIndex++;
+				}
+
 				for (const auto & pair : UniformVarNames)
 				{
 					std::string OldUniformTypeName = "uniform type_" + pair.first;
@@ -2945,15 +2987,15 @@ TSharedPtr<ANSICHAR> FOpenGLFrontend::PrepareCodeForOfflineCompilation(const GLS
 	bool bNeedsExtDrawInstancedDefine = false;
 	if (Capabilities.TargetPlatform == EPlatformType::Android || Capabilities.TargetPlatform == EPlatformType::Web)
 	{
-		const TCHAR *ES310Version = TEXT("#version 310 es");
-		StrOutSource.Append(ES310Version);
+		const TCHAR *ES320Version = TEXT("#version 320 es");
+		StrOutSource.Append(ES320Version);
 		StrOutSource.Append(TEXT("\n"));
-		OriginalShaderSource.RemoveFromStart(ES310Version);
+		OriginalShaderSource.RemoveFromStart(ES320Version);
 	}
 
 	const GLenum TypeEnum = GLFrequencyTable[Frequency];
 	// The incoming glsl may have preprocessor code that is dependent on defines introduced via the engine.
-	// This is the place to insert such engine preprocessor defines, immediately after the glsl version declaration.
+	// This is the place to insert such engine preprocessor defines, immediately after the glsl version declarati
 
 	if (bEmitMobileMultiView)
 	{
diff --git a/Engine/Source/Runtime/OpenGLDrv/Private/Android/AndroidOpenGL.h b/Engine/Source/Runtime/OpenGLDrv/Private/Android/AndroidOpenGL.h
--- a/Engine/Source/Runtime/OpenGLDrv/Private/Android/AndroidOpenGL.h
+++ b/Engine/Source/Runtime/OpenGLDrv/Private/Android/AndroidOpenGL.h
@@ -158,7 +158,7 @@ struct FAndroidOpenGL : public FOpenGLES
 	static FORCEINLINE GLint GetFirstComputeUAVUnit()			{ return 0; }
 	static FORCEINLINE GLint GetMaxComputeUAVUnits()			{ check(MaxComputeUAVUnits != -1); return MaxComputeUAVUnits; }
 	static FORCEINLINE GLint GetFirstVertexUAVUnit()			{ return 0; }
-	static FORCEINLINE GLint GetFirstPixelUAVUnit()				{ return 4; }
+	static FORCEINLINE GLint GetFirstPixelUAVUnit()				{ return 0; }
 	static FORCEINLINE GLint GetMaxPixelUAVUnits()				{ check(MaxPixelUAVUnits != -1); return MaxPixelUAVUnits; }
 	static FORCEINLINE GLint GetMaxCombinedUAVUnits()			{ return MaxCombinedUAVUnits; }
 		
diff --git a/Engine/Source/Runtime/OpenGLDrv/Private/OpenGLShaders.cpp b/Engine/Source/Runtime/OpenGLDrv/Private/OpenGLShaders.cpp
--- a/Engine/Source/Runtime/OpenGLDrv/Private/OpenGLShaders.cpp
+++ b/Engine/Source/Runtime/OpenGLDrv/Private/OpenGLShaders.cpp
@@ -855,11 +855,11 @@ void OPENGLDRV_API GLSLToDeviceCompatibleGLSL(FAnsiCharArray& GlslCodeOriginal,
 	
 	if (Capabilities.TargetPlatform == EOpenGLShaderTargetPlatform::OGLSTP_Android)
 	{
-		const ANSICHAR* ES310Version = "#version 310 es";
+		const ANSICHAR* ES320Version = "#version 320 es";
 
-		AppendCString(GlslCode, ES310Version);
+		AppendCString(GlslCode, ES320Version);
 		AppendCString(GlslCode, "\n");
-		ReplaceCString(GlslCodeOriginal, ES310Version, "");
+		ReplaceCString(GlslCodeOriginal, ES320Version, "");
 	}
 
 	if (TypeEnum == GL_FRAGMENT_SHADER && Capabilities.bRequiresDisabledEarlyFragmentTests)
diff --git a/Engine/Source/Runtime/OpenGLDrv/Public/OpenGL4.h b/Engine/Source/Runtime/OpenGLDrv/Public/OpenGL4.h
--- a/Engine/Source/Runtime/OpenGLDrv/Public/OpenGL4.h
+++ b/Engine/Source/Runtime/OpenGLDrv/Public/OpenGL4.h
@@ -107,7 +107,7 @@ struct FOpenGL4 : public FOpenGL3
 	static FORCEINLINE GLint GetFirstComputeUAVUnit()			{ return 0; }
 	static FORCEINLINE GLint GetMaxComputeUAVUnits()			{ check(MaxComputeUAVUnits != -1); return MaxComputeUAVUnits; }
 	static FORCEINLINE GLint GetFirstVertexUAVUnit()			{ return 0; }
-	static FORCEINLINE GLint GetFirstPixelUAVUnit()				{ return 4; }
+	static FORCEINLINE GLint GetFirstPixelUAVUnit()				{ return 0; }
 	static FORCEINLINE GLint GetMaxPixelUAVUnits()				{ check(MaxPixelUAVUnits != -1); return MaxPixelUAVUnits; }
 	static FORCEINLINE GLint GetMaxCombinedUAVUnits()			{ return MaxCombinedUAVUnits; }
 		
