Update hlslcc GL backend to match GetFirstPixelUAVUnit

diff --git a/Engine/Source/Developer/ShaderFormatOpenGL/Private/GlslBackend.cpp b/Engine/Source/Developer/ShaderFormatOpenGL/Private/GlslBackend.cpp
--- a/Engine/Source/Developer/ShaderFormatOpenGL/Private/GlslBackend.cpp
+++ b/Engine/Source/Developer/ShaderFormatOpenGL/Private/GlslBackend.cpp
@@ -1158,7 +1158,7 @@ class ir_gen_glsl_visitor : public ir_visitor
 				{
 					0, //vertex_shader, must match FOpenGL::GetFirstVertexUAVUnit()
 					0,
-					4, //fragment_shader, must match FOpenGL::GetFirstPixelUAVUnit()
+					0, //fragment_shader, must match FOpenGL::GetFirstPixelUAVUnit()
 					0,
 					0,
 					0, //compute_shader
diff --git a/Engine/Source/Developer/ShaderFormatOpenGL/Private/ShaderFormatOpenGL.cpp b/Engine/Source/Developer/ShaderFormatOpenGL/Private/ShaderFormatOpenGL.cpp
--- a/Engine/Source/Developer/ShaderFormatOpenGL/Private/ShaderFormatOpenGL.cpp
+++ b/Engine/Source/Developer/ShaderFormatOpenGL/Private/ShaderFormatOpenGL.cpp
@@ -17,7 +17,7 @@ class FShaderFormatGLSL : public IShaderFormat
 	enum
 	{
 		/** Version for shader format, this becomes part of the DDC key. */
-		UE_SHADER_GLSL_VER = 98,
+		UE_SHADER_GLSL_VER = 99,
 	};
 
 	void CheckFormat(FName Format) const
