Merging CL 17736980

//UE5/Release-5.0/...

to //UE5/Main/...

diff --git a/Engine/Source/Developer/ShaderCompilerCommon/Private/ShaderConductorContext.cpp b/Engine/Source/Developer/ShaderCompilerCommon/Private/ShaderConductorContext.cpp
--- a/Engine/Source/Developer/ShaderCompilerCommon/Private/ShaderConductorContext.cpp
+++ b/Engine/Source/Developer/ShaderCompilerCommon/Private/ShaderConductorContext.cpp
@@ -918,6 +918,7 @@ namespace CrossCompiler
 			/*InputAttribute:*/					"in.var.ATTRIBUTE",
 			/*GlobalsUniformBuffer:*/			"$Globals",
 			/*IntermediateTessControlOutput:*/	"temp.var.hullMainRetVal",
+			/*DummySampler:*/					"SPIRV_Cross_DummySampler",
 		};
 		return IdentifierTable;
 	}
diff --git a/Engine/Source/Developer/ShaderCompilerCommon/Public/ShaderConductorContext.h b/Engine/Source/Developer/ShaderCompilerCommon/Public/ShaderConductorContext.h
--- a/Engine/Source/Developer/ShaderCompilerCommon/Public/ShaderConductorContext.h
+++ b/Engine/Source/Developer/ShaderCompilerCommon/Public/ShaderConductorContext.h
@@ -85,6 +85,9 @@ namespace CrossCompiler
 
 		/** Identifier for the intermediate output variable in a tessellation-control shader. */
 		const ANSICHAR* IntermediateTessControlOutput;
+
+		/** Identifier for dummy samplers used for platforms where samplers are required */
+		const ANSICHAR* DummySampler;
 	};
 
 	/** Wrapper class to handle interface between UE and ShaderConductor. Use to compile HLSL shaders to SPIR-V or high-level languages such as Metal. */
diff --git a/Engine/Source/Developer/ShaderFormatOpenGL/Private/OpenGLShaderCompiler.cpp b/Engine/Source/Developer/ShaderFormatOpenGL/Private/OpenGLShaderCompiler.cpp
--- a/Engine/Source/Developer/ShaderFormatOpenGL/Private/OpenGLShaderCompiler.cpp
+++ b/Engine/Source/Developer/ShaderFormatOpenGL/Private/OpenGLShaderCompiler.cpp
@@ -1877,6 +1877,10 @@ static bool CompileToGlslWithShaderConductor(
 			check(SPVRResult == SPV_REFLECT_RESULT_SUCCESS);
 		}
 
+		// Spirv-cross can add a dummy sampler for glsl which isn't in the reflection bindings
+		const ANSICHAR* SPIRV_DummySamplerName = CrossCompiler::FShaderConductorContext::GetIdentifierTable().DummySampler;
+		Samplers.Add(SPIRV_DummySamplerName);
+
 		{
 			uint32 Count = 0;
 			SPVRResult = Reflection.EnumeratePushConstantBlocks(&Count, nullptr);
@@ -2502,6 +2506,7 @@ static bool CompileToGlslWithShaderConductor(
 						std::string SamplerName = "SPIRV_Cross_Combined";
 						SamplerName  += TCHAR_TO_ANSI(*(Texture + Sampler));
 						size_t FindCombinedSampler = GlslSource.find(SamplerName.c_str());
+
 						if (FindCombinedSampler != std::string::npos)
 						{
 							checkf(SamplerPos != std::string::npos, TEXT("Generated GLSL shader is expected to have combined sampler '%s:%s' but no appropriate 'uniform' declaration could be found"), *Texture, *Sampler);
@@ -2516,7 +2521,12 @@ static bool CompileToGlslWithShaderConductor(
 							NewDefine += "\n";
 							GlslSource.insert(SamplerPos+1, NewDefine);
 
-							UsedSamplers.Add(Sampler);
+							// Do not add an entry for the dummy sampler as it will throw errors in the runtime checks
+							if (Sampler != SPIRV_DummySamplerName)
+							{
+								UsedSamplers.Add(Sampler);
+							}
+							
 							SamplerString += FString::Printf(TEXT("%s%s"), SamplerString.Len() ? TEXT(",") : TEXT(""), *Sampler);
 						}
 					}
@@ -2592,7 +2602,6 @@ static bool CompileToGlslWithShaderConductor(
 					size_t SBufferEndPos = GlslSource.find(";", SBufferPos);
 					
 					std::string ReplacementSubStr = GlslSource.substr(SBufferPos + 2, SBufferEndPos - SBufferPos - 2);
-					UE_LOG(LogOpenGLShaderCompiler, Warning, TEXT("Found ReplacementSubStr %s"), UTF8_TO_TCHAR(ReplacementSubStr.c_str()));
 					GlslSource.erase(SBufferPos + 1, SBufferEndPos - SBufferPos - 1);
 
 					const std::string SBufferData = "_m0";
