Fix for packed matrix arrays in DXC

diff --git a/Engine/Source/Developer/ShaderFormatOpenGL/Private/OpenGLShaderCompiler.cpp b/Engine/Source/Developer/ShaderFormatOpenGL/Private/OpenGLShaderCompiler.cpp
--- a/Engine/Source/Developer/ShaderFormatOpenGL/Private/OpenGLShaderCompiler.cpp
+++ b/Engine/Source/Developer/ShaderFormatOpenGL/Private/OpenGLShaderCompiler.cpp
@@ -1342,7 +1342,7 @@ void AddMemberToPackedUB(const std::string & FrequencyPrefix,
 		}
 
 		Name += "(Offset)";
-		if (type.op == SpvOpTypeMatrix)
+		if (type.op == SpvOpTypeMatrix || (type.traits.numeric.matrix.column_count == 4 && type.traits.numeric.matrix.row_count == 4))
 		{
 			OffsetString += " + (int(Offset) * 4)";
 		}
@@ -1364,7 +1364,7 @@ void AddMemberToPackedUB(const std::string & FrequencyPrefix,
 		UniformPrefix += "u";
 	}
 
-	if (type.op == SpvOpTypeMatrix)
+	if (type.op == SpvOpTypeMatrix || (type.traits.numeric.matrix.column_count == 4 && type.traits.numeric.matrix.row_count == 4))
 	{
 		if ((type.traits.numeric.matrix.column_count == 4 && type.traits.numeric.matrix.row_count == 4))
 		{
