Fix for runtime shader error on Android with multiview enabled

diff --git a/Engine/Source/Developer/ShaderFormatOpenGL/Private/OpenGLShaderCompiler.cpp b/Engine/Source/Developer/ShaderFormatOpenGL/Private/OpenGLShaderCompiler.cpp
--- a/Engine/Source/Developer/ShaderFormatOpenGL/Private/OpenGLShaderCompiler.cpp
+++ b/Engine/Source/Developer/ShaderFormatOpenGL/Private/OpenGLShaderCompiler.cpp
@@ -2160,10 +2160,14 @@ static bool CompileToGlslWithShaderConductor(
 		default:
 			TargetDesc.CompileFlags.SetDefine(TEXT("force_flattened_io_blocks"), 1);
 			TargetDesc.CompileFlags.SetDefine(TEXT("emit_uniform_buffer_as_plain_uniforms"), 1);
-			if (Frequency == SF_Vertex)
+
+			// If we have mobile multiview define set then set the view count and enable extension
+			const FString* MultiViewDefine = Input.Environment.GetDefinitions().Find(TEXT("MOBILE_MULTI_VIEW"));
+			if (Frequency == SF_Vertex && MultiViewDefine && *MultiViewDefine == "1")
 			{
 				TargetDesc.CompileFlags.SetDefine(TEXT("ovr_multiview_view_count"), 2);
 			}
+
 			TargetDesc.Language = CrossCompiler::EShaderConductorLanguage::Essl;
 			TargetDesc.Version = 320;
 			break;
