Fix hlslcc depth fetch support on Android devices not supporting GL_ARM_shader_framebuffer_fetch_depth_stencil
Reapply change that did not correctly make it into 19.01.

diff --git a/Engine/Source/Developer/ShaderFormatOpenGL/Private/GlslBackend.cpp b/Engine/Source/Developer/ShaderFormatOpenGL/Private/GlslBackend.cpp
--- a/Engine/Source/Developer/ShaderFormatOpenGL/Private/GlslBackend.cpp
+++ b/Engine/Source/Developer/ShaderFormatOpenGL/Private/GlslBackend.cpp
@@ -1249,7 +1249,7 @@ class ir_gen_glsl_visitor : public ir_visitor
 				}
 				
 				const bool bIsDepthTarget = var->name && bUsesDepthbufferFetch && strncmp(var->name, "out_Target1", 11) == 0;
-				const bool bNeedsFBFOutput = (bUsesFrameBufferFetch || bUsesDepthbufferFetch) && var->name && (strncmp(var->name, "out_Target0", 11) == 0);
+				const bool bNeedsFBFOutput = bIsDepthTarget || (var->name && (bUsesFrameBufferFetch && strncmp(var->name, "out_Target0", 11) == 0));
 				const char* StorageQualifier = bNeedsFBFOutput ? FBF_StorageQualifier : mode_str[var->mode];
 
 				ralloc_asprintf_append(
@@ -1358,7 +1358,7 @@ class ir_gen_glsl_visitor : public ir_visitor
 				{
 					ralloc_asprintf_append(buffer, "  layout(location=1) inout float out_Target1;\n");
 				}
-				ralloc_asprintf_append(buffer, "  float DepthbufferFetchES2() { return out_Target1.r; }\n");
+				ralloc_asprintf_append(buffer, "  float DepthbufferFetchES2() { return out_Target1; }\n");
 				ralloc_asprintf_append(buffer, "#else\n");
 				ralloc_asprintf_append(buffer, "  float DepthbufferFetchES2() { return 0.0; }\n");
 				ralloc_asprintf_append(buffer, "#endif\n\n");
diff --git a/Engine/Source/Developer/ShaderFormatOpenGL/Private/ShaderFormatOpenGL.cpp b/Engine/Source/Developer/ShaderFormatOpenGL/Private/ShaderFormatOpenGL.cpp
--- a/Engine/Source/Developer/ShaderFormatOpenGL/Private/ShaderFormatOpenGL.cpp
+++ b/Engine/Source/Developer/ShaderFormatOpenGL/Private/ShaderFormatOpenGL.cpp
@@ -17,7 +17,7 @@ class FShaderFormatGLSL : public IShaderFormat
 	enum
 	{
 		/** Version for shader format, this becomes part of the DDC key. */
-		UE_SHADER_GLSL_VER = 102,
+		UE_SHADER_GLSL_VER = 103,
 	};
 
 	void CheckFormat(FName Format) const
