Added support for framebuffer fetch and pixel local storage to OpenGLES with DXC
(Submitting for Carl.Lloyd)

diff --git a/Engine/Shaders/Private/MobileBasePassPixelShader.usf b/Engine/Shaders/Private/MobileBasePassPixelShader.usf
--- a/Engine/Shaders/Private/MobileBasePassPixelShader.usf
+++ b/Engine/Shaders/Private/MobileBasePassPixelShader.usf
@@ -587,12 +587,14 @@ void Main(
 	, in float4 SvPosition : SV_Position
 	OPTIONAL_IsFrontFace
 #if DEFERRED_SHADING_PATH
-#if !USE_GLES_FBF_DEFERRED
+#if USE_GLES_FBF_DEFERRED
+	, out HALF4_TYPE OutProxy : SV_Target0
+#else
 	, out HALF4_TYPE OutColor : SV_Target0
+#endif
 	, out HALF4_TYPE OutGBufferA : SV_Target1
 	, out HALF4_TYPE OutGBufferB : SV_Target2
-	, out HALF4_TYPE OutGBufferC	: SV_Target3  
-#endif
+	, out HALF4_TYPE OutGBufferC : SV_Target3  
 #else
 	, out HALF4_TYPE OutColor : SV_Target0
 #endif
@@ -603,16 +605,11 @@ void Main(
 	, out float OutDepth : SV_Depth
 #endif
 	)
-{
-
+{  
 #if DEFERRED_SHADING_PATH && USE_GLES_FBF_DEFERRED
 	half4 OutColor;	
-	half4 OutGBufferA;
-	half4 OutGBufferB;
-	half4 OutGBufferC;
-	float OutSceneDepthAux;
 #endif
-  
+
 #if MOBILE_MULTI_VIEW
 	ResolvedView = ResolveView(uint(BasePassInterpolants.MultiViewId));
 #else
@@ -804,13 +801,6 @@ void Main(
 	GBuffer.PrecomputedShadowFactors.r = Shadow;
 
 	EncodeGBuffer(GBuffer, OutGBufferA, OutGBufferB, OutGBufferC, OutGBufferD, OutGBufferE, OutGBufferF, OutGBufferVelocity);
-
-#if USE_GLES_FBF_DEFERRED
-	GLFramebufferFetchWrite1(OutGBufferA);  
-	GLFramebufferFetchWrite2(OutGBufferB);  
-	GLFramebufferFetchWrite3(OutGBufferC);
-	GLFramebufferFetchDepthWrite(SvPosition.z);   
-#endif
 #else
 
 #if ENABLE_SKY_LIGHT
@@ -1121,6 +1111,6 @@ void Main(
 #endif
 
 #if DEFERRED_SHADING_PATH && USE_GLES_FBF_DEFERRED 
-	GLFramebufferFetchWrite0(half4(OutColor.rgb, 1.0f));
+	OutProxy.rgb = OutColor.rgb;
 #endif
 }
diff --git a/Engine/Shaders/Private/MobileDeferredShading.usf b/Engine/Shaders/Private/MobileDeferredShading.usf
--- a/Engine/Shaders/Private/MobileDeferredShading.usf
+++ b/Engine/Shaders/Private/MobileDeferredShading.usf
@@ -31,6 +31,8 @@
 #include "/Engine/Public/Platform/Metal/MetalSubpassSupport.ush"
 #elif VULKAN_PROFILE
 #include "/Engine/Public/Platform/Vulkan/VulkanSubpassSupport.ush"
+#elif OPENGL_PROFILE
+#include "/Engine/Public/Platform/GL/GLSubpassSupport.ush"
 #endif
 
 #ifndef IS_SPOT_LIGHT 
@@ -71,11 +73,11 @@ void FetchGBuffer(in float2 UV, out float4 GBufferA, out float4 GBufferB, out fl
 	GBufferD = 0; 
 	SceneDepth = ConvertFromDeviceZ(SubpassFetchR_4());
 #elif USE_GLES_FBF_DEFERRED
-	GBufferA = GLFramebufferFetch1(); 
-	GBufferB = GLFramebufferFetch2(); 
-	GBufferC = GLFramebufferFetch3();  
+	GBufferA = GLSubpassFetch1(); 
+	GBufferB = GLSubpassFetch2(); 
+	GBufferC = GLSubpassFetch3();  
 	GBufferD = 0;   
-	SceneDepth = ConvertFromDeviceZ(GLFramebufferFetchDepth());
+	SceneDepth = ConvertFromDeviceZ(DepthbufferFetchES2());
 #else
 	GBufferA = Texture2DSampleLevel(MobileSceneTextures.GBufferATexture, MobileSceneTextures.GBufferATextureSampler, UV, 0); 
 	GBufferB = Texture2DSampleLevel(MobileSceneTextures.GBufferBTexture, MobileSceneTextures.GBufferBTextureSampler, UV, 0);
@@ -234,11 +236,19 @@ void GetLightGridLocalLighting(
 		GetDirectLighting(LightData, ShadingModelContext, GBuffer, TranslatedWorldPosition, CameraVector, MobileLightAccumulator);
 	}
 }
-
+  
 void MobileDirectLightPS(
 	noperspective float4 UVAndScreenPos : TEXCOORD0, 
 	float4 SvPosition : SV_POSITION, 
-	out HALF4_TYPE OutColor : SV_Target0)
+#if USE_GLES_FBF_DEFERRED
+	out HALF4_TYPE OutProxyAdditive : SV_Target0,
+	out HALF4_TYPE OutGBufferA : SV_Target1,
+	out HALF4_TYPE OutGBufferB : SV_Target2,
+	out HALF4_TYPE OutGBufferC : SV_Target3
+#else
+	out HALF4_TYPE OutColor : SV_Target0
+#endif
+)
 {
 	ResolvedView = ResolveView();
 
@@ -326,7 +336,7 @@ void MobileDirectLightPS(
 		// Planar reflections win over reflection environment
 		SpecularIBL = lerp(SpecularIBL, PlanarReflection.rgb, PlanarReflection.a);
 	}
-
+	  
 	MobileLightAccumulator_Add(MobileLightAccumulator, 0.0f, SpecularIBL * ShadingModelContext.EnvBrdf, 1.0f);
 #elif APPLY_SKY_REFLECTION
 	float SkyAverageBrightness = 1.0f;
@@ -354,15 +364,26 @@ void MobileDirectLightPS(
 
 	// MobileHDR applies PreExposure in tonemapper
 	LightAttenuation *= View.PreExposure;
-					
+
+#if defined(MOBILE_DEFERRED_SHADING) && USE_GLES_FBF_DEFERRED == 1
+	OutProxyAdditive.rgb = Color.rgb * LightAttenuation;
+#else
 	OutColor.rgb = Color.rgb * LightAttenuation;
 	OutColor.a = 1;
+#endif
 }
 
 void MobileRadialLightPS(
 	float4 InScreenPosition : TEXCOORD0,
 	float4 SVPos			: SV_POSITION,
-	out float4 OutColor		: SV_Target0
+#if USE_GLES_FBF_DEFERRED
+	out HALF4_TYPE OutProxyAdditive : SV_Target0,
+	out HALF4_TYPE OutGBufferA : SV_Target1,
+	out HALF4_TYPE OutGBufferB : SV_Target2,
+	out HALF4_TYPE OutGBufferC : SV_Target3
+#else
+	out HALF4_TYPE OutColor : SV_Target0
+#endif
 )
 {
 	FGBufferData GBuffer = (FGBufferData)0;
@@ -410,6 +431,10 @@ void MobileRadialLightPS(
 	// MobileHDR applies PreExposure in tonemapper
 	LightAttenuation*= View.PreExposure;
 
-	OutColor.rgb = Color * LightAttenuation;
+#if defined(MOBILE_DEFERRED_SHADING) && USE_GLES_FBF_DEFERRED == 1
+	OutProxyAdditive.rgb = Color.rgb * LightAttenuation;
+#else
+	OutColor.rgb = Color.rgb * LightAttenuation;
 	OutColor.a = 1;
+#endif
 }
diff --git a/Engine/Shaders/Private/MobileDeferredUtils.usf b/Engine/Shaders/Private/MobileDeferredUtils.usf
--- a/Engine/Shaders/Private/MobileDeferredUtils.usf
+++ b/Engine/Shaders/Private/MobileDeferredUtils.usf
@@ -1,13 +1,17 @@
 // Copyright Epic Games, Inc. All Rights Reserved.
 #include "Common.ush"
 
+#if USE_GLES_FBF_DEFERRED
+#include "/Engine/Public/Platform/GL/GLSubpassSupport.ush"
+#endif
+
 void MobileDeferredCopyPLSPS(
 	noperspective float4 UVAndScreenPos : TEXCOORD0, 
 	float4 SvPosition : SV_POSITION, 
 	out half4 OutColor : SV_Target0)
 {
 #if USE_GLES_FBF_DEFERRED
-	OutColor = float4(GLFramebufferFetch0(), 1.0f);
+	OutColor = half4(GLSubpassFetch0().rgb, 1.0);
 #else
 	OutColor = 1.0f;  
 #endif
@@ -19,7 +23,7 @@ void MobileDeferredCopyDepthPS(
 	out float OutSceneDepthAux : SV_Target0)
 {
 #if USE_GLES_FBF_DEFERRED
-	OutSceneDepthAux = GLFramebufferFetchDepth();
+	OutSceneDepthAux = DepthbufferFetchES2();
 #else
 	OutSceneDepthAux = 1.0f;  
 #endif
diff --git a/Engine/Shaders/Private/SceneTexturesCommon.ush b/Engine/Shaders/Private/SceneTexturesCommon.ush
--- a/Engine/Shaders/Private/SceneTexturesCommon.ush
+++ b/Engine/Shaders/Private/SceneTexturesCommon.ush
@@ -8,8 +8,12 @@
 
 #include "Common.ush"
 
-#if VULKAN_PROFILE && COMPILER_DXC && PIXELSHADER
-#	include "/Engine/Public/Platform/Vulkan/VulkanSubpassSupport.ush"
+#if COMPILER_DXC && PIXELSHADER
+#if VULKAN_PROFILE
+	#include "/Engine/Public/Platform/Vulkan/VulkanSubpassSupport.ush"
+#elif OPENGL_PROFILE
+	#include "/Engine/Public/Platform/GL/GLSubpassSupport.ush"
+#endif
 #endif
 
 // Return far plane when scene textures are disabled in order not to break depthfade
diff --git a/Engine/Shaders/Public/Platform/GL/GLCommon.ush b/Engine/Shaders/Public/Platform/GL/GLCommon.ush
--- a/Engine/Shaders/Public/Platform/GL/GLCommon.ush
+++ b/Engine/Shaders/Public/Platform/GL/GLCommon.ush
@@ -5,6 +5,6 @@
 =============================================================================*/
 
 // Update this GUID to improve shader recompilation for GL only shaders
-// GUID = 58EB247E-D403-451E-B31F-4B24D38898B8
+// GUID = E121241A-1AD3-416B-8F74-288560D5F0005
 
 #pragma once
diff --git a/Engine/Shaders/Public/Platform/GL/GLSubpassSupport.ush b/Engine/Shaders/Public/Platform/GL/GLSubpassSupport.ush
new file mode 100644
--- /dev/null
+++ b/Engine/Shaders/Public/Platform/GL/GLSubpassSupport.ush
@@ -0,0 +1,32 @@
+// Copyright Epic Games, Inc. All Rights Reserved.
+
+/*=============================================================================
+	GLSubpassSupport.usf: Vulkan specific subpass intrinsics 
+=============================================================================*/
+#pragma once
+
+#if COMPILER_DXC && PIXELSHADER
+
+	// Note: Do not rename subpass input attachment declarations, or adjust your changes in VulkanBackend.cpp (see VULKAN_SUBPASS_FETCH_VAR_W)
+
+	[[vk::input_attachment_index(0)]]
+	SubpassInput<float> GENERATED_SubpassDepthFetchAttachment;
+	#define GLSubpassDepthFetch() GENERATED_SubpassDepthFetchAttachment.SubpassLoad( )
+
+	[[vk::input_attachment_index(1)]]
+	SubpassInput<float4> GENERATED_SubpassFetchAttachment0;
+	#define GLSubpassFetch0() GENERATED_SubpassFetchAttachment0.SubpassLoad()
+
+	[[vk::input_attachment_index(2)]]
+	SubpassInput<float4> GENERATED_SubpassFetchAttachment1;
+	#define GLSubpassFetch1() GENERATED_SubpassFetchAttachment1.SubpassLoad()
+
+	[[vk::input_attachment_index(3)]]
+	SubpassInput<float4> GENERATED_SubpassFetchAttachment2;
+	#define GLSubpassFetch2() GENERATED_SubpassFetchAttachment2.SubpassLoad()
+	
+	[[vk::input_attachment_index(4)]]
+	SubpassInput<float4> GENERATED_SubpassFetchAttachment3;
+	#define GLSubpassFetch3() GENERATED_SubpassFetchAttachment3.SubpassLoad()
+	
+#endif // COMPILER_DXC && PIXELSHADER
diff --git a/Engine/Source/Developer/ShaderCompilerCommon/Private/ShaderCompilerCommon.cpp b/Engine/Source/Developer/ShaderCompilerCommon/Private/ShaderCompilerCommon.cpp
--- a/Engine/Source/Developer/ShaderCompilerCommon/Private/ShaderCompilerCommon.cpp
+++ b/Engine/Source/Developer/ShaderCompilerCommon/Private/ShaderCompilerCommon.cpp
@@ -1921,6 +1921,11 @@ void DumpDebugShaderText(const FShaderCompilerInput& Input, ANSICHAR* InSource,
 	DumpDebugShaderBinary(Input, InSource, InSourceLength * sizeof(ANSICHAR), FileExtension);
 }
 
+void DumpDebugShaderText(const FShaderCompilerInput& Input, ANSICHAR* InSource, int32 InSourceLength, const FString& FileName, const FString& FileExtension)
+{
+	DumpDebugShaderBinary(Input, InSource, InSourceLength * sizeof(ANSICHAR), FileName, FileExtension);
+}
+
 void DumpDebugShaderBinary(const FShaderCompilerInput& Input, void* InData, int32 InDataByteSize, const FString& FileExtension)
 {
 	if (InData != nullptr && InDataByteSize > 0 && !FileExtension.IsEmpty())
@@ -1934,6 +1939,19 @@ void DumpDebugShaderBinary(const FShaderCompilerInput& Input, void* InData, int3
 	}
 }
 
+void DumpDebugShaderBinary(const FShaderCompilerInput& Input, void* InData, int32 InDataByteSize, const FString& FileName, const FString& FileExtension)
+{
+	if (InData != nullptr && InDataByteSize > 0 && !FileExtension.IsEmpty())
+	{
+		const FString Filename = Input.DumpDebugInfoPath / FileName + TEXT(".") + FileExtension;
+		if (TUniquePtr<FArchive> FileWriter = TUniquePtr<FArchive>(IFileManager::Get().CreateFileWriter(*Filename)))
+		{
+			FileWriter->Serialize(InData, InDataByteSize);
+			FileWriter->Close();
+		}
+	}
+}
+
 static void DumpDebugShaderDisassembled(const FShaderCompilerInput& Input, CrossCompiler::EShaderConductorIR Language, void* InData, int32 InDataByteSize, const FString& FileExtension)
 {
 	if (InData != nullptr && InDataByteSize > 0 && !FileExtension.IsEmpty())
diff --git a/Engine/Source/Developer/ShaderCompilerCommon/Private/ShaderConductorContext.cpp b/Engine/Source/Developer/ShaderCompilerCommon/Private/ShaderConductorContext.cpp
--- a/Engine/Source/Developer/ShaderCompilerCommon/Private/ShaderConductorContext.cpp
+++ b/Engine/Source/Developer/ShaderCompilerCommon/Private/ShaderConductorContext.cpp
@@ -348,6 +348,9 @@ namespace CrossCompiler
 		// Convert FString to ANSI string and store them as intermediates
 		FMemory::Memzero(OutTargetDesc);
 
+		Intermediates.Flags.Empty();
+		Intermediates.FlagRefs.Empty();
+
 		switch (InTarget.Language)
 		{
 		case EShaderConductorLanguage::Hlsl:
diff --git a/Engine/Source/Developer/ShaderCompilerCommon/Public/ShaderCompilerCommon.h b/Engine/Source/Developer/ShaderCompilerCommon/Public/ShaderCompilerCommon.h
--- a/Engine/Source/Developer/ShaderCompilerCommon/Public/ShaderCompilerCommon.h
+++ b/Engine/Source/Developer/ShaderCompilerCommon/Public/ShaderCompilerCommon.h
@@ -183,8 +183,10 @@ extern SHADERCOMPILERCOMMON_API void DumpDebugUSF(const FShaderCompilerInput& In
 extern SHADERCOMPILERCOMMON_API void DumpDebugUSF(const FShaderCompilerInput& Input, const FString& Source, uint32 HlslCCFlags = 0, const TCHAR* OverrideBaseFilename = nullptr);
 
 extern SHADERCOMPILERCOMMON_API void DumpDebugShaderText(const FShaderCompilerInput& Input, const FString& InSource, const FString& FileExtension);
-extern SHADERCOMPILERCOMMON_API void DumpDebugShaderText(const FShaderCompilerInput& Input, ANSICHAR* InSource, int32 InSourceLength, const FString& FileExtension);
+extern SHADERCOMPILERCOMMON_API void DumpDebugShaderText(const FShaderCompilerInput& Input, ANSICHAR* InSource, int32 InSourceLength, const FString& FileExtension); 
+extern SHADERCOMPILERCOMMON_API void DumpDebugShaderText(const FShaderCompilerInput& Input, ANSICHAR* InSource, int32 InSourceLength, const FString& FileName, const FString& FileExtension);
 extern SHADERCOMPILERCOMMON_API void DumpDebugShaderBinary(const FShaderCompilerInput& Input, void* InData, int32 InDataByteSize, const FString& FileExtension);
+extern SHADERCOMPILERCOMMON_API void DumpDebugShaderBinary(const FShaderCompilerInput& Input, void* InData, int32 InDataByteSize, const FString& FileName, const FString& FileExtension);
 extern SHADERCOMPILERCOMMON_API void DumpDebugShaderDisassembledSpirv(const FShaderCompilerInput& Input, void* InData, int32 InDataByteSize, const FString& FileExtension);
 extern SHADERCOMPILERCOMMON_API void DumpDebugShaderDisassembledDxil(const FShaderCompilerInput& Input, void* InData, int32 InDataByteSize, const FString& FileExtension);
 
diff --git a/Engine/Source/Developer/ShaderFormatOpenGL/Private/OpenGLShaderCompiler.cpp b/Engine/Source/Developer/ShaderFormatOpenGL/Private/OpenGLShaderCompiler.cpp
--- a/Engine/Source/Developer/ShaderFormatOpenGL/Private/OpenGLShaderCompiler.cpp
+++ b/Engine/Source/Developer/ShaderFormatOpenGL/Private/OpenGLShaderCompiler.cpp
@@ -1292,14 +1292,14 @@ void GetSpvVarQualifier(const SpvReflectBlockVariable& Member, FString & Out)
 }
 
 // Adds a member to ne included in the PackedUB structures and generates the #define for readability in glsl
-void AddMemberToPackedUB(const std::string & FrequencyPrefix,
-							const SpvReflectBlockVariable& Member,
-							const std::string& UBName,
-							int32_t Index,
-							TMap<FString, uint32>& Offsets,
-							TArray<PackedUBMemberInfo>& MemberInfos,
-							TArray<std::string>& Remap,
-							TArray<std::string>& Arrays)
+void AddMemberToPackedUB(const std::string& FrequencyPrefix,
+	const SpvReflectBlockVariable& Member,
+	const std::string& UBName,
+	int32_t Index,
+	TMap<FString, uint32>& Offsets,
+	TArray<PackedUBMemberInfo>& MemberInfos,
+	TArray<std::string>& Remap,
+	TArray<std::string>& Arrays)
 {
 	std::string ArrayName;
 	const uint32 MbrSize = Member.size / sizeof(float);
@@ -1312,7 +1312,7 @@ void AddMemberToPackedUB(const std::string & FrequencyPrefix,
 	std::replace(SanitizedName.begin(), SanitizedName.end(), '.', '_');
 
 	uint32& Offset = Offsets.FindOrAdd(TypeQualifier);
-	
+
 	auto const type = *Member.type_description;
 
 	bool const bArray = type.traits.array.dims_count > 0;
@@ -1335,7 +1335,7 @@ void AddMemberToPackedUB(const std::string & FrequencyPrefix,
 	MemberInfo.SrcOffset = Member.offset;
 	MemberInfo.DestOffset = Offset * 4 * sizeof(float);
 	MemberInfo.SrcSizeInFloats = Member.size / sizeof(float);
-	MemberInfo.DestSizeInFloats = Member.size / sizeof(float);	
+	MemberInfo.DestSizeInFloats = Member.size / sizeof(float);
 
 	if (bArray)
 	{
@@ -1425,22 +1425,21 @@ void AddMemberToPackedUB(const std::string & FrequencyPrefix,
 		Name += TCHAR_TO_UTF8(*TypeQualifier);
 		Name += "[";
 		Name += OffsetString;
-		Name += "].";
+		Name += "]";
 		switch (type.traits.numeric.vector.component_count)
 		{
 		case 0:
 		case 1:
-			Name += "x";
+			Name += ".x";
 			break;
 		case 2:
-			Name += "xy";
+			Name += ".xy";
 			break;
 		case 3:
-			Name += "xyz";
+			Name += ".xyz";
 			break;
 		case 4:
 		default:
-			Name += "xyzw";
 			break;
 		}
 	}
@@ -1456,7 +1455,7 @@ void AddMemberToPackedUB(const std::string & FrequencyPrefix,
 	Offset += Align(MbrSize, 4) / 4;
 }
 
-void GetPackedUniformString(std::string& OutputString, const std::string& UniformPrefix, const FString & Key, uint32_t Index)
+void GetPackedUniformString(std::string& OutputString, const std::string& UniformPrefix, const FString& Key, uint32_t Index)
 {
 	if (Key == TEXT("u"))
 	{
@@ -1492,616 +1491,386 @@ void GetPackedUniformString(std::string& OutputString, const std::string& Unifor
 	}
 }
 
-static bool CompileToGlslWithShaderConductor(
-	const FShaderCompilerInput&	Input,
-	FShaderCompilerOutput&		Output,
-	const FString&				WorkingDirectory,
-	GLSLVersion					Version,
-	const EShaderFrequency		Frequency,
-	uint32						CCFlags,
-	const FString&				PreprocessedShader,
-	char*&						OutGlslShaderSource)
+struct ReflectionData
 {
-	CrossCompiler::FShaderConductorContext CompilerContext;
-
-	const bool bDumpDebugInfo = (Input.DumpDebugInfoPath != TEXT("") && IFileManager::Get().DirectoryExists(*Input.DumpDebugInfoPath));
-	const bool bRewriteHlslSource = true;
-
-	// Initialize compilation options for ShaderConductor
-	CrossCompiler::FShaderConductorOptions Options;
-	Options.bDisableScalarBlockLayout = true;
-	Options.bRemapAttributeLocations = true;
-	Options.bPreserveStorageInput = true;
-	
-	// Convert input strings from FString to ANSI strings
-	std::string SourceData(TCHAR_TO_UTF8(*PreprocessedShader));
-	std::string FileName(TCHAR_TO_UTF8(*Input.VirtualSourceFilePath));
-	std::string EntryPointName(TCHAR_TO_UTF8(*Input.EntryPointName));
+	TArray<FString> Textures;
+	TArray<FString> Samplers;
+	TArray<std::string> UAVs;
+	TArray<std::string> StructuredBuffers;
 	
-	bool bEmulatedUBs = Input.Environment.CompilerFlags.Contains(CFLAG_UseEmulatedUB);
+	std::map<std::string, std::string> UniformVarNames;
 
-	// HLSL framebuffer declarations. Used to modify HLSL input source.
-	const ANSICHAR* HlslFrameBufferDeclarations =
-		"float4 gl_FragColor;\n"
-		"float4 gl_LastFragColorARM;\n"
-		"float gl_LastFragDepthARM;\n"
-		"bool ARM_shader_framebuffer_fetch;\n"
-		"bool ARM_shader_framebuffer_fetch_depth_stencil;\n"
-		"float4 FramebufferFetchES2()\n"
-		"{\n"
-		"  if (!ARM_shader_framebuffer_fetch)\n"
-		"  {\n"
-		"    return gl_FragColor;\n"
-		"  }\n"
-		"  else\n"
-		"  {\n"
-		"    return gl_LastFragColorARM;\n"
-		"  }\n"
-		"}\n"
-		"float DepthbufferFetchES2()\n"
-		"{\n"
-		"  return (ARM_shader_framebuffer_fetch_depth_stencil == 0 ? 0.0 : gl_LastFragDepthARM);\n"
-		"}\n"
-		;
+	TMap<FString, uint32> GlobalOffsets;
+	TArray<std::string> GlobalRemap;
+	TArray<std::string> GlobalArrays;
+	TArray<PackedUBMemberInfo> GlobalMemberInfos;
 
-	SourceData = HlslFrameBufferDeclarations + SourceData;
-
-	const ANSICHAR* GlslFrameBufferExtensions =
-		"\n\n#ifdef GL_ARM_shader_framebuffer_fetch_depth_stencil\n"
-		"\t#extension GL_ARM_shader_framebuffer_fetch_depth_stencil : enable\n"
-		"#elif defined(GL_EXT_shader_framebuffer_fetch)\n"
-		"\t#extension GL_EXT_shader_framebuffer_fetch : enable\n"
-		"\t#define FBF_STORAGE_QUALIFIER inout\n"
-		"#endif\n"
-		"#extension GL_EXT_texture_buffer : enable\n"
-		"// end extensions";
-
-	// GLSL framebuffer macro definitions. Used to patch GLSL output source.
-	const ANSICHAR* GlslFrameBufferDefines =
-		"\n\n#ifdef UE_EXT_shader_framebuffer_fetch\n"
-		"#define _Globals_ARM_shader_framebuffer_fetch 0\n"
-		"#define FRAME_BUFFERFETCH_STORAGE_QUALIFIER inout\n"
-		"#define _Globals_gl_FragColor out_var_SV_Target0\n"
-		"#define _Globals_gl_LastFragColorARM vec4(0.0, 0.0, 0.0, 0.0)\n"
-		"#elif defined( GL_ARM_shader_framebuffer_fetch)\n"
-		"#define _Globals_ARM_shader_framebuffer_fetch 1\n"
-		"#define FRAME_BUFFERFETCH_STORAGE_QUALIFIER out\n"
-		"#define _Globals_gl_FragColor vec4(0.0, 0.0, 0.0, 0.0)\n"
-		"#define _Globals_gl_LastFragColorARM gl_LastFragDepthARM\n"
-		"#else\n"
-		"#define FRAME_BUFFERFETCH_STORAGE_QUALIFIER out\n"
-		"#define _Globals_ARM_shader_framebuffer_fetch 0\n"
-		"#define _Globals_gl_FragColor vec4(0.0, 0.0, 0.0, 0.0)\n"
-		"#define _Globals_gl_LastFragColorARM vec4(0.0, 0.0, 0.0, 0.0)\n"
-		"#endif\n"
-		"#ifdef GL_ARM_shader_framebuffer_fetch_depth_stencil\n"
-		"#define _Globals_ARM_shader_framebuffer_fetch_depth_stencil 1u\n"
-		"#else\n"
-		"#define _Globals_ARM_shader_framebuffer_fetch_depth_stencil 0u\n"
-		"#endif\n"
-		;
-	
-	// Inject additional macro definitions to circumvent missing features: external textures
-	FShaderCompilerDefinitions AdditionalDefines;
-	AdditionalDefines.SetDefine(TEXT("TextureExternal"), TEXT("Texture2D"));
+	TMap<uint32, TMap<FString, uint32>> PackedUBOffsets;
+	TMap<uint32, TArray<std::string>> PackedUBRemap;
+	TMap<uint32, TArray<std::string>> PackedUBArrays;
+	TMap<uint32, TArray<PackedUBMemberInfo>> PackedUBMemberInfos;
+	TMap<uint32, std::string> PackedUBNames;
 
-	if (bDumpDebugInfo)
-	{
-		FString DirectCompileLine = CrossCompiler::CreateResourceTableFromEnvironment(Input.Environment);
+	TArray<FString> InputVarNames;
+	TArray<FString> OutputVarNames;
+};
 
-		DirectCompileLine += TEXT("#if 0 /*DIRECT COMPILE*/\n");
-		DirectCompileLine += CreateShaderCompilerWorkerDirectCommandLine(Input, CCFlags);
-		DirectCompileLine += TEXT("\n#endif /*DIRECT COMPILE*/\n");
+void ParseReflectionData(CrossCompiler::FHlslccHeaderWriter& CCHeaderWriter, ReflectionData& ReflectionOut, TArray<uint32>& SpirvData, spv_reflect::ShaderModule& Reflection,
+							const ANSICHAR* SPIRV_DummySamplerName, const EShaderFrequency Frequency, bool bEmulatedUBs)
+{
+	const ANSICHAR* FrequencyPrefix = GetFrequencyPrefix(Frequency);
 
-		DumpDebugUSF(Input, (PreprocessedShader + DirectCompileLine), CCFlags);
+	check(Reflection.GetResult() == SPV_REFLECT_RESULT_SUCCESS);
 
-		if (Input.bGenerateDirectCompileFile)
-		{
-			FFileHelper::SaveStringToFile(CreateShaderCompilerWorkerDirectCommandLine(Input), *(Input.DumpDebugInfoPath / TEXT("DirectCompile.txt")));
-		}
-	}
+	SpvReflectResult SPVRResult = SPV_REFLECT_RESULT_NOT_READY;
+	TArray<SpvReflectBlockVariable*> ConstantBindings;
+	const uint32 GlobalSetId = 32;
 
-	// Load shader source into compiler context
-	CompilerContext.LoadSource(SourceData.c_str(), FileName.c_str(), EntryPointName.c_str(), Frequency, &AdditionalDefines);
+	FSpirvReflectBindings ReflectionBindings;
+	ReflectionBindings.GatherDescriptorBindings(Reflection);
 
-	bool bCompilationFailed = false;
+	uint32 BufferIndices = 0xffffffff;
+	uint32 UAVIndices = 0xffffffff;
+	uint32 TextureIndices = 0xffffffff;
+	uint32 UBOIndices = 0xffffffff;
+	uint32 SamplerIndices = 0xffffffff;
+	uint32_t PackedUBIndex = 0;
 
-	if (bRewriteHlslSource)
+	for (auto const& Binding : ReflectionBindings.TBufferUAVs)
 	{
-		// Rewrite HLSL source code to remove unused global resources and variables
-		Options.bRemoveUnusedGlobals = true;
-		if (CompilerContext.RewriteHlsl(Options))
-		{
-			// Adopt new rewritten shader source
-			SourceData = CompilerContext.GetSourceString();
-
-			if (bDumpDebugInfo)
-			{
-				DumpDebugShaderText(Input, ANSI_TO_TCHAR(SourceData.c_str()), TEXT("rewritten.hlsl"));
-			}
-		}
-		else
-		{
-			CompilerContext.FlushErrors(Output.Errors);
-			bCompilationFailed = true;
-		}
-		Options.bRemoveUnusedGlobals = false;
-	}
+		check(UAVIndices);
+		uint32 Index = FPlatformMath::CountTrailingZeros(UAVIndices);
 
-	// Compile HLSL source to SPIR-V binary
-	TArray<uint32> SpirvData;
-	
-	if (!bCompilationFailed && !CompilerContext.CompileHlslToSpirv(Options, SpirvData))
-	{
-		// Flush compile errors
-		CompilerContext.FlushErrors(Output.Errors);
-		bCompilationFailed = true;
-	}
+		// UAVs always claim all slots so we don't have conflicts as D3D expects 0-7
+		BufferIndices &= ~(1 << Index);
+		TextureIndices &= ~(1llu << uint64(Index));
+		UAVIndices &= ~(1 << Index);
 
-	if (!bCompilationFailed)
-	{
-		const ANSICHAR* FrequencyPrefix = GetFrequencyPrefix(Frequency);
+		CCHeaderWriter.WriteUAV(UTF8_TO_TCHAR(Binding->name), Index);
 
-		// Now perform reflection on the SPIRV and tweak any decorations that we need to.
-		// This used to be done via JSON, but that was slow and alloc happy so use SPIRV-Reflect instead.
-		spv_reflect::ShaderModule Reflection(SpirvData.Num() * sizeof(uint32), SpirvData.GetData());
+		SPVRResult = Reflection.ChangeDescriptorBindingNumbers(Binding, Index, GlobalSetId);
+		check(SPVRResult == SPV_REFLECT_RESULT_SUCCESS);
 
-		check(Reflection.GetResult() == SPV_REFLECT_RESULT_SUCCESS);
+		ReflectionOut.UAVs.Add(Binding->name);
+	}
 
-		SpvReflectResult SPVRResult = SPV_REFLECT_RESULT_NOT_READY;
-		TArray<SpvReflectBlockVariable*> ConstantBindings;
-		const uint32 GlobalSetId = 32;
+	for (auto const& Binding : ReflectionBindings.SBufferUAVs)
+	{
+		check(UAVIndices);
+		uint32 Index = FPlatformMath::CountTrailingZeros(UAVIndices);
 
-		CrossCompiler::FHlslccHeaderWriter CCHeaderWriter;
-		TArray<FString> Textures;
-		TArray<FString> Samplers;
-		TArray<std::string> UAVs;
-		TArray<std::string> StructuredBuffers;
-		uint32 UAVIndices = 0xffffffff;
-		uint32 BufferIndices = 0xffffffff;
-		uint32 TextureIndices = 0xffffffff;
-		uint32 UBOIndices = 0xffffffff;
-		uint32 SamplerIndices = 0xffffffff;
+		// UAVs always claim all slots so we don't have conflicts as D3D expects 0-7
+		BufferIndices &= ~(1 << Index);
+		TextureIndices &= ~(1llu << uint64(Index));
+		UAVIndices &= ~(1 << Index);
 
-		std::map<std::string, std::string> UniformVarNames;
+		CCHeaderWriter.WriteUAV(UTF8_TO_TCHAR(Binding->name), Index);
+		ReflectionOut.StructuredBuffers.Add(Binding->name);
 
-		FSpirvReflectBindings ReflectionBindings;
-		ReflectionBindings.GatherDescriptorBindings(Reflection);
+		SPVRResult = Reflection.ChangeDescriptorBindingNumbers(Binding, Index, GlobalSetId);
+		check(SPVRResult == SPV_REFLECT_RESULT_SUCCESS);
+	}
 
-		for (auto const& Binding : ReflectionBindings.TBufferUAVs)
-		{
-			check(UAVIndices);
-			uint32 Index = FPlatformMath::CountTrailingZeros(UAVIndices);
+	for (auto const& Binding : ReflectionBindings.TextureUAVs)
+	{
+		check(UAVIndices);
+		uint32 Index = FPlatformMath::CountTrailingZeros(UAVIndices);
 
-			// UAVs always claim all slots so we don't have conflicts as D3D expects 0-7
-			BufferIndices &= ~(1 << Index);
-			TextureIndices &= ~(1llu << uint64(Index));
-			UAVIndices &= ~(1 << Index);
+		// UAVs always claim all slots so we don't have conflicts as D3D expects 0-7
+		// For texture2d this allows us to emulate atomics with buffers
+		BufferIndices &= ~(1 << Index);
+		TextureIndices &= ~(1llu << uint64(Index));
+		UAVIndices &= ~(1 << Index);
 
-			CCHeaderWriter.WriteUAV(UTF8_TO_TCHAR(Binding->name), Index);
+		CCHeaderWriter.WriteUAV(UTF8_TO_TCHAR(Binding->name), Index);
 
-			SPVRResult = Reflection.ChangeDescriptorBindingNumbers(Binding, Index, GlobalSetId);
-			check(SPVRResult == SPV_REFLECT_RESULT_SUCCESS);
+		SPVRResult = Reflection.ChangeDescriptorBindingNumbers(Binding, Index, GlobalSetId);
+		check(SPVRResult == SPV_REFLECT_RESULT_SUCCESS);
 
-			UAVs.Add(Binding->name);
-		}
+		ReflectionOut.UAVs.Add(Binding->name);
+	}
 
-		for (auto const& Binding : ReflectionBindings.SBufferUAVs)
-		{
-			check(UAVIndices);
-			uint32 Index = FPlatformMath::CountTrailingZeros(UAVIndices);
+	for (auto const& Binding : ReflectionBindings.TBufferSRVs)
+	{
+		check(TextureIndices);
+		uint32 Index = FPlatformMath::CountTrailingZeros(TextureIndices);
 
-			// UAVs always claim all slots so we don't have conflicts as D3D expects 0-7
-			BufferIndices &= ~(1 << Index);
-			TextureIndices &= ~(1llu << uint64(Index));
-			UAVIndices &= ~(1 << Index);
+		// No support for 3-component types in dxc/SPIRV/MSL - need to expose my workarounds there too
+		TextureIndices &= ~(1llu << uint64(Index));
 
-			CCHeaderWriter.WriteUAV(UTF8_TO_TCHAR(Binding->name), Index);
-			StructuredBuffers.Add(Binding->name);
+		ReflectionOut.Textures.Add(UTF8_TO_TCHAR(Binding->name));
 
-			SPVRResult = Reflection.ChangeDescriptorBindingNumbers(Binding, Index, GlobalSetId);
-			check(SPVRResult == SPV_REFLECT_RESULT_SUCCESS);
-		}
+		SPVRResult = Reflection.ChangeDescriptorBindingNumbers(Binding, Index, GlobalSetId);
+		check(SPVRResult == SPV_REFLECT_RESULT_SUCCESS);
+	}
 
-		for (auto const& Binding : ReflectionBindings.TextureUAVs)
-		{
-			check(UAVIndices);
-			uint32 Index = FPlatformMath::CountTrailingZeros(UAVIndices);
+	for (int32_t SBufferIndex = ReflectionBindings.SBufferSRVs.Num() - 1; SBufferIndex >= 0; SBufferIndex--)
+	{
+		auto Binding = ReflectionBindings.SBufferSRVs[SBufferIndex];
+		check(BufferIndices);
+		uint32 Index = FPlatformMath::CountTrailingZeros(BufferIndices);
 
-			// UAVs always claim all slots so we don't have conflicts as D3D expects 0-7
-			// For texture2d this allows us to emulate atomics with buffers
-			BufferIndices &= ~(1 << Index);
-			TextureIndices &= ~(1llu << uint64(Index));
-			UAVIndices &= ~(1 << Index);
+		BufferIndices &= ~(1 << Index);
 
-			CCHeaderWriter.WriteUAV(UTF8_TO_TCHAR(Binding->name), Index);
+		SPVRResult = Reflection.ChangeDescriptorBindingNumbers(Binding, Index, GlobalSetId);
+		check(SPVRResult == SPV_REFLECT_RESULT_SUCCESS);
 
-			SPVRResult = Reflection.ChangeDescriptorBindingNumbers(Binding, Index, GlobalSetId);
-			check(SPVRResult == SPV_REFLECT_RESULT_SUCCESS);
+		CCHeaderWriter.WriteUAV(UTF8_TO_TCHAR(Binding->name), Index);
+		ReflectionOut.StructuredBuffers.Add(Binding->name);
+	}
 
-			UAVs.Add(Binding->name);
-		}
+	for (auto const& Binding : ReflectionBindings.UniformBuffers)
+	{
+		check(UBOIndices);
+		uint32 Index = FPlatformMath::CountTrailingZeros(UBOIndices);
+		UBOIndices &= ~(1 << Index);
 
-		for (auto const& Binding : ReflectionBindings.TBufferSRVs)
+		// Global uniform buffer - handled specially as we care about the internal layout
+		if (strstr(Binding->name, "$Globals"))
 		{
-			check(TextureIndices);
-			uint32 Index = FPlatformMath::CountTrailingZeros(TextureIndices);
-
-			// No support for 3-component types in dxc/SPIRV/MSL - need to expose my workarounds there too
-			TextureIndices &= ~(1llu << uint64(Index));
-
-			Textures.Add(UTF8_TO_TCHAR(Binding->name));
-
-			SPVRResult = Reflection.ChangeDescriptorBindingNumbers(Binding, Index, GlobalSetId);
-			check(SPVRResult == SPV_REFLECT_RESULT_SUCCESS);
-		}
+			TMap<FString, uint32_t> GlobalOffsetSizes;
+			for (uint32 i = 0; i < Binding->block.member_count; i++)
+			{
+				SpvReflectBlockVariable& member = Binding->block.members[i];
 
-		for (int32_t SBufferIndex = ReflectionBindings.SBufferSRVs.Num() - 1; SBufferIndex >= 0; SBufferIndex--)
-		{
-			auto Binding = ReflectionBindings.SBufferSRVs[SBufferIndex];
-			check(BufferIndices);
-			uint32 Index = FPlatformMath::CountTrailingZeros(BufferIndices);
+				FString TypeQualifier;
+				GetSpvVarQualifier(member, TypeQualifier);
 
-			BufferIndices &= ~(1 << Index);
+				uint32_t& GlobalOffsetSize = GlobalOffsetSizes.FindOrAdd(TypeQualifier);
 
-			SPVRResult = Reflection.ChangeDescriptorBindingNumbers(Binding, Index, GlobalSetId);
-			check(SPVRResult == SPV_REFLECT_RESULT_SUCCESS);
+				if (strstr(member.name, "gl_") || !strcmp(member.name, "ARM_shader_framebuffer_fetch") || !strcmp(member.name, "ARM_shader_framebuffer_fetch_depth_stencil"))
+				{
+					continue;
+				}
 
-			CCHeaderWriter.WriteUAV(UTF8_TO_TCHAR(Binding->name), Index);
-			StructuredBuffers.Add(Binding->name);
+				bool bHalfPrecision = member.decoration_flags & SPV_REFLECT_DECORATION_RELAXED_PRECISION;
+				CCHeaderWriter.WritePackedGlobal(ANSI_TO_TCHAR(member.name), CrossCompiler::FHlslccHeaderWriter::EncodePackedGlobalType(*(member.type_description), bHalfPrecision), GlobalOffsetSize, member.size);
+				GlobalOffsetSize += Align(member.size, 16);
+
+				AddMemberToPackedUB(FrequencyPrefix,
+					member,
+					"_Globals_",
+					-1,
+					ReflectionOut.GlobalOffsets,
+					ReflectionOut.GlobalMemberInfos,
+					ReflectionOut.GlobalRemap,
+					ReflectionOut.GlobalArrays);
+			}
 		}
-
-		TMap<FString, uint32> GlobalOffsets;
-		TArray<std::string> GlobalRemap;
-		TArray<std::string> GlobalArrays;
-		TArray<PackedUBMemberInfo> GlobalMemberInfos;
-
-		TMap<uint32, TMap<FString, uint32>> PackedUBOffsets;
-		TMap<uint32, TArray<std::string>> PackedUBRemap;
-		TMap<uint32, TArray<std::string>> PackedUBArrays;
-		TMap<uint32, TArray<PackedUBMemberInfo>> PackedUBMemberInfos;
-		TMap<uint32, std::string> PackedUBNames;
-
-		uint32_t PackedUBIndex = 0;
-
-		for (auto const& Binding : ReflectionBindings.UniformBuffers)
+		else
 		{
-			check(UBOIndices);
-			uint32 Index = FPlatformMath::CountTrailingZeros(UBOIndices);
-			UBOIndices &= ~(1 << Index);
-
-			// Global uniform buffer - handled specially as we care about the internal layout
-			if (strstr(Binding->name, "$Globals"))
+			if (bEmulatedUBs)
 			{
-				TMap<FString, uint32_t> GlobalOffsetSizes;
+				//CCHeaderWriter.WritePackedUB(Binding->binding);
+
+				ReflectionOut.PackedUBOffsets.Add(PackedUBIndex);
+				ReflectionOut.PackedUBRemap.Add(PackedUBIndex);
+				ReflectionOut.PackedUBArrays.Add(PackedUBIndex);
+				ReflectionOut.PackedUBMemberInfos.Add(PackedUBIndex);
+				ReflectionOut.PackedUBNames.Add(PackedUBIndex, std::string(Binding->name));
+
 				for (uint32 i = 0; i < Binding->block.member_count; i++)
 				{
 					SpvReflectBlockVariable& member = Binding->block.members[i];
 
-					FString TypeQualifier;
-					GetSpvVarQualifier(member, TypeQualifier);
-
-					uint32_t& GlobalOffsetSize = GlobalOffsetSizes.FindOrAdd(TypeQualifier);
+					std::string UBName = "_" + std::string(Binding->name) + "_";
 
-					if (strstr(member.name, "gl_") || !strcmp(member.name, "ARM_shader_framebuffer_fetch") || !strcmp(member.name, "ARM_shader_framebuffer_fetch_depth_stencil"))
+					if (member.type_description->type_flags & SPV_REFLECT_TYPE_FLAG_STRUCT)
 					{
-						continue;
-					}
+						for (uint32 n = 0; n < member.member_count; n++)
+						{
+							// Clone struct member and rename it for struct
+							SpvReflectBlockVariable StructMember = member.members[n];
 
-					bool bHalfPrecision = member.decoration_flags & SPV_REFLECT_DECORATION_RELAXED_PRECISION;
-					CCHeaderWriter.WritePackedGlobal(ANSI_TO_TCHAR(member.name), CrossCompiler::FHlslccHeaderWriter::EncodePackedGlobalType(*(member.type_description), bHalfPrecision), GlobalOffsetSize, member.size);
-					GlobalOffsetSize += Align(member.size, 16);
-
-					AddMemberToPackedUB(FrequencyPrefix,
-						member,
-						"_Globals_",
-						-1,
-						GlobalOffsets,
-						GlobalMemberInfos,
-						GlobalRemap,
-						GlobalArrays);
-				}
-			}
-			else
-			{
-				if (bEmulatedUBs)
-				{
-					//CCHeaderWriter.WritePackedUB(Binding->binding);
-					
-					PackedUBOffsets.Add(PackedUBIndex);
-					PackedUBRemap.Add(PackedUBIndex);
-					PackedUBArrays.Add(PackedUBIndex);
-					PackedUBMemberInfos.Add(PackedUBIndex);
-					PackedUBNames.Add(PackedUBIndex, std::string(Binding->name));
-
-					for (uint32 i = 0; i < Binding->block.member_count; i++)
-					{
-						SpvReflectBlockVariable& member = Binding->block.members[i];
+							std::string StructMemberName = std::string(member.name) + "." + std::string(StructMember.name);
+							StructMember.name = StructMemberName.c_str();
 
-						std::string UBName = "_" + std::string(Binding->name) + "_";
-						
-						if (member.type_description->type_flags & SPV_REFLECT_TYPE_FLAG_STRUCT)
-						{
-							for (uint32 n = 0; n < member.member_count; n++)
-							{
-								// Clone struct member and rename it for struct
-								SpvReflectBlockVariable StructMember = member.members[n];
-								
-								std::string StructMemberName = std::string(member.name) + "." + std::string(StructMember.name);
-								StructMember.name = StructMemberName.c_str();
-
-								AddMemberToPackedUB(FrequencyPrefix,
-									StructMember,
-									UBName,
-									PackedUBIndex,
-									PackedUBOffsets[PackedUBIndex],
-									PackedUBMemberInfos[PackedUBIndex],
-									PackedUBRemap[PackedUBIndex],
-									PackedUBArrays[PackedUBIndex]);
-							}
-						}
-						else
-						{
 							AddMemberToPackedUB(FrequencyPrefix,
-								member,
+								StructMember,
 								UBName,
 								PackedUBIndex,
-								PackedUBOffsets[PackedUBIndex],
-								PackedUBMemberInfos[PackedUBIndex],
-								PackedUBRemap[PackedUBIndex],
-								PackedUBArrays[PackedUBIndex]);
+								ReflectionOut.PackedUBOffsets[PackedUBIndex],
+								ReflectionOut.PackedUBMemberInfos[PackedUBIndex],
+								ReflectionOut.PackedUBRemap[PackedUBIndex],
+								ReflectionOut.PackedUBArrays[PackedUBIndex]);
 						}
 					}
-
-					PackedUBIndex++;
-				}
-				else
-				{
-					std::string OldName = Binding->name;
-					std::string NewName = FrequencyPrefix;
-					NewName += "b";
-					NewName += std::to_string(Index);
-					UniformVarNames[OldName] = NewName;
-					// Regular uniform buffer - we only care about the binding index
-					CCHeaderWriter.WriteUniformBlock(UTF8_TO_TCHAR(Binding->name), Index);
+					else
+					{
+						AddMemberToPackedUB(FrequencyPrefix,
+							member,
+							UBName,
+							PackedUBIndex,
+							ReflectionOut.PackedUBOffsets[PackedUBIndex],
+							ReflectionOut.PackedUBMemberInfos[PackedUBIndex],
+							ReflectionOut.PackedUBRemap[PackedUBIndex],
+							ReflectionOut.PackedUBArrays[PackedUBIndex]);
+					}
 				}
-			}
 
-			SPVRResult = Reflection.ChangeDescriptorBindingNumbers(Binding, Index, GlobalSetId);
-			check(SPVRResult == SPV_REFLECT_RESULT_SUCCESS);
+				PackedUBIndex++;
+			}
+			else
+			{
+				std::string OldName = Binding->name;
+				std::string NewName = FrequencyPrefix;
+				NewName += "b";
+				NewName += std::to_string(Index);
+				ReflectionOut.UniformVarNames[OldName] = NewName;
+				// Regular uniform buffer - we only care about the binding index
+				CCHeaderWriter.WriteUniformBlock(UTF8_TO_TCHAR(Binding->name), Index);
+			}
 		}
 
-		WritePackedUBHeader(CCHeaderWriter, PackedUBMemberInfos, PackedUBNames);
+		SPVRResult = Reflection.ChangeDescriptorBindingNumbers(Binding, Index, GlobalSetId);
+		check(SPVRResult == SPV_REFLECT_RESULT_SUCCESS);
+	}
 
-		for (auto const& Binding : ReflectionBindings.TextureSRVs)
-		{
-			check(TextureIndices);
-			uint32 Index = FPlatformMath::CountTrailingZeros64(TextureIndices);
-			TextureIndices &= ~(1llu << uint64(Index));
+	WritePackedUBHeader(CCHeaderWriter, ReflectionOut.PackedUBMemberInfos, ReflectionOut.PackedUBNames);
 
-			Textures.Add(UTF8_TO_TCHAR(Binding->name));
+	for (auto const& Binding : ReflectionBindings.TextureSRVs)
+	{
+		check(TextureIndices);
+		uint32 Index = FPlatformMath::CountTrailingZeros64(TextureIndices);
+		TextureIndices &= ~(1llu << uint64(Index));
 
-			SPVRResult = Reflection.ChangeDescriptorBindingNumbers(Binding, Index, GlobalSetId);
-			check(SPVRResult == SPV_REFLECT_RESULT_SUCCESS);
-		}
+		ReflectionOut.Textures.Add(UTF8_TO_TCHAR(Binding->name));
 
-		for (auto const& Binding : ReflectionBindings.Samplers)
-		{
-			check(SamplerIndices);
-			uint32 Index = FPlatformMath::CountTrailingZeros(SamplerIndices);
-			SamplerIndices &= ~(1 << Index);
+		SPVRResult = Reflection.ChangeDescriptorBindingNumbers(Binding, Index, GlobalSetId);
+		check(SPVRResult == SPV_REFLECT_RESULT_SUCCESS);
+	}
 
-			Samplers.Add(UTF8_TO_TCHAR(Binding->name));
+	for (auto const& Binding : ReflectionBindings.Samplers)
+	{
+		check(SamplerIndices);
+		uint32 Index = FPlatformMath::CountTrailingZeros(SamplerIndices);
+		SamplerIndices &= ~(1 << Index);
 
-			SPVRResult = Reflection.ChangeDescriptorBindingNumbers(Binding, Index, GlobalSetId);
-			check(SPVRResult == SPV_REFLECT_RESULT_SUCCESS);
-		}
+		ReflectionOut.Samplers.Add(UTF8_TO_TCHAR(Binding->name));
 
-		// Spirv-cross can add a dummy sampler for glsl which isn't in the reflection bindings
-		const ANSICHAR* SPIRV_DummySamplerName = CrossCompiler::FShaderConductorContext::GetIdentifierTable().DummySampler;
-		Samplers.Add(SPIRV_DummySamplerName);
+		SPVRResult = Reflection.ChangeDescriptorBindingNumbers(Binding, Index, GlobalSetId);
+		check(SPVRResult == SPV_REFLECT_RESULT_SUCCESS);
+	}
 
+	// Spirv-cross can add a dummy sampler for glsl which isn't in the reflection bindings
+	ReflectionOut.Samplers.Add(SPIRV_DummySamplerName);
+
+	{
+		uint32 Count = 0;
+		SPVRResult = Reflection.EnumeratePushConstantBlocks(&Count, nullptr);
+		check(SPVRResult == SPV_REFLECT_RESULT_SUCCESS);
+		ConstantBindings.SetNum(Count);
+		SPVRResult = Reflection.EnumeratePushConstantBlocks(&Count, ConstantBindings.GetData());
+		check(SPVRResult == SPV_REFLECT_RESULT_SUCCESS);
+		if (Count > 0)
 		{
-			uint32 Count = 0;
-			SPVRResult = Reflection.EnumeratePushConstantBlocks(&Count, nullptr);
-			check(SPVRResult == SPV_REFLECT_RESULT_SUCCESS);
-			ConstantBindings.SetNum(Count);
-			SPVRResult = Reflection.EnumeratePushConstantBlocks(&Count, ConstantBindings.GetData());
-			check(SPVRResult == SPV_REFLECT_RESULT_SUCCESS);
-			if (Count > 0)
+			for (auto const& Var : ConstantBindings)
 			{
-				for (auto const& Var : ConstantBindings)
+				// Global uniform buffer - handled specially as we care about the internal layout
+				if (strstr(Var->name, "$Globals"))
 				{
-					// Global uniform buffer - handled specially as we care about the internal layout
-					if (strstr(Var->name, "$Globals"))
-					{
-						TMap<FString, uint32_t> GlobalOffsetSizes;
-
-						for (uint32 i = 0; i < Var->member_count; i++)
-						{
-							SpvReflectBlockVariable& member = Var->members[i];
-						
-							FString TypeQualifier;
-							GetSpvVarQualifier(member, TypeQualifier);
+					TMap<FString, uint32_t> GlobalOffsetSizes;
 
-							uint32_t& GlobalOffsetSize = GlobalOffsetSizes.FindOrAdd(TypeQualifier);
+					for (uint32 i = 0; i < Var->member_count; i++)
+					{
+						SpvReflectBlockVariable& member = Var->members[i];
 
-							if (!strcmp(member.name, "gl_FragColor") || !strcmp(member.name, "gl_LastFragColorARM") || !strcmp(member.name, "gl_LastFragDepthARM") || !strcmp(member.name, "ARM_shader_framebuffer_fetch") || !strcmp(member.name, "ARM_shader_framebuffer_fetch_depth_stencil"))
-							{
-								continue;
-							}
+						FString TypeQualifier;
+						GetSpvVarQualifier(member, TypeQualifier);
 
-							bool bHalfPrecision = member.decoration_flags & SPV_REFLECT_DECORATION_RELAXED_PRECISION;
-							CCHeaderWriter.WritePackedGlobal(ANSI_TO_TCHAR(member.name), CrossCompiler::FHlslccHeaderWriter::EncodePackedGlobalType(*(member.type_description), bHalfPrecision), GlobalOffsetSize, member.size);
-							GlobalOffsetSize += Align(member.size, 16);
+						uint32_t& GlobalOffsetSize = GlobalOffsetSizes.FindOrAdd(TypeQualifier);
 
-							AddMemberToPackedUB(FrequencyPrefix,
-								member,
-								"_Globals_",
-								-1,
-								GlobalOffsets,
-								GlobalMemberInfos,
-								GlobalRemap,
-								GlobalArrays);
+						if (!strcmp(member.name, "gl_FragColor") || !strcmp(member.name, "gl_LastFragColorARM") || !strcmp(member.name, "gl_LastFragDepthARM") || !strcmp(member.name, "ARM_shader_framebuffer_fetch") || !strcmp(member.name, "ARM_shader_framebuffer_fetch_depth_stencil"))
+						{
+							continue;
 						}
+
+						bool bHalfPrecision = member.decoration_flags & SPV_REFLECT_DECORATION_RELAXED_PRECISION;
+						CCHeaderWriter.WritePackedGlobal(ANSI_TO_TCHAR(member.name), CrossCompiler::FHlslccHeaderWriter::EncodePackedGlobalType(*(member.type_description), bHalfPrecision), GlobalOffsetSize, member.size);
+						GlobalOffsetSize += Align(member.size, 16);
+
+						AddMemberToPackedUB(FrequencyPrefix,
+							member,
+							"_Globals_",
+							-1,
+							ReflectionOut.GlobalOffsets,
+							ReflectionOut.GlobalMemberInfos,
+							ReflectionOut.GlobalRemap,
+							ReflectionOut.GlobalArrays);
 					}
 				}
 			}
 		}
+	}
 
-		TArray<FString> OutputVarNames;
+	{
+		uint32 AssignedInputs = 0;
 
+		ReflectionBindings.GatherOutputAttributes(Reflection);
+		for (SpvReflectInterfaceVariable* Var : ReflectionBindings.OutputAttributes)
 		{
-			uint32 AssignedInputs = 0;
-
-			ReflectionBindings.GatherOutputAttributes(Reflection);
-			for (SpvReflectInterfaceVariable* Var : ReflectionBindings.OutputAttributes)
+			if (Var->storage_class == SpvStorageClassOutput && Var->built_in == -1 && !CrossCompiler::FShaderConductorContext::IsIntermediateSpirvOutputVariable(Var->name))
 			{
-				if (Var->storage_class == SpvStorageClassOutput && Var->built_in == -1 && !CrossCompiler::FShaderConductorContext::IsIntermediateSpirvOutputVariable(Var->name))
+				if (Frequency == SF_Pixel && strstr(Var->name, "SV_Target"))
 				{
-					if (Frequency == SF_Pixel && strstr(Var->name, "SV_Target"))
-					{
-						FString TypeQualifier;
-
-						auto const type = *Var->type_description;
-						uint32_t masked_type = type.type_flags & 0xF;
+					FString TypeQualifier;
 
-						switch (masked_type) {
-						default: checkf(false, TEXT("unsupported component type %d"), masked_type); break;
-						case SPV_REFLECT_TYPE_FLAG_BOOL: TypeQualifier = TEXT("b"); break;
-						case SPV_REFLECT_TYPE_FLAG_INT: TypeQualifier = (type.traits.numeric.scalar.signedness ? TEXT("i") : TEXT("u")); break;
-						case SPV_REFLECT_TYPE_FLAG_FLOAT: TypeQualifier = (type.traits.numeric.scalar.width == 32 ? TEXT("f") : TEXT("h")); break;
-						}
+					auto const type = *Var->type_description;
+					uint32_t masked_type = type.type_flags & 0xF;
 
-						if (type.type_flags & SPV_REFLECT_TYPE_FLAG_MATRIX)
-						{
-							TypeQualifier += FString::Printf(TEXT("%d%d"), type.traits.numeric.matrix.row_count, type.traits.numeric.matrix.column_count);
-						}
-						else if (type.type_flags & SPV_REFLECT_TYPE_FLAG_VECTOR)
-						{
-							TypeQualifier += FString::Printf(TEXT("%d"), type.traits.numeric.vector.component_count);
-						}
-						else
-						{
-							TypeQualifier += TEXT("1");
-						}
+					switch (masked_type) {
+					default: checkf(false, TEXT("unsupported component type %d"), masked_type); break;
+					case SPV_REFLECT_TYPE_FLAG_BOOL: TypeQualifier = TEXT("b"); break;
+					case SPV_REFLECT_TYPE_FLAG_INT: TypeQualifier = (type.traits.numeric.scalar.signedness ? TEXT("i") : TEXT("u")); break;
+					case SPV_REFLECT_TYPE_FLAG_FLOAT: TypeQualifier = (type.traits.numeric.scalar.width == 32 ? TEXT("f") : TEXT("h")); break;
+					}
 
-						FString Name = ANSI_TO_TCHAR(Var->name);
-						Name.ReplaceInline(TEXT("."), TEXT("_"));
-						OutputVarNames.Add(Name);
-						CCHeaderWriter.WriteOutputAttribute(TEXT("out_Target"), *TypeQualifier, Var->location, /*bLocationPrefix:*/ true, /*bLocationSuffix:*/ true);
+					if (type.type_flags & SPV_REFLECT_TYPE_FLAG_MATRIX)
+					{
+						TypeQualifier += FString::Printf(TEXT("%d%d"), type.traits.numeric.matrix.row_count, type.traits.numeric.matrix.column_count);
+					}
+					else if (type.type_flags & SPV_REFLECT_TYPE_FLAG_VECTOR)
+					{
+						TypeQualifier += FString::Printf(TEXT("%d"), type.traits.numeric.vector.component_count);
 					}
 					else
 					{
-						unsigned Location = Var->location;
-						unsigned SemanticIndex = Location;
-						check(Var->semantic);
-						unsigned i = (unsigned)strlen(Var->semantic);
-						check(i);
-						while (isdigit((unsigned char)(Var->semantic[i - 1])))
-						{
-							i--;
-						}
-						if (i < strlen(Var->semantic))
-						{
-							SemanticIndex = (unsigned)atoi(Var->semantic + i);
-							if (Location != SemanticIndex)
-							{
-								Location = SemanticIndex;
-							}
-						}
+						TypeQualifier += TEXT("1");
+					}
 
-						while ((1 << Location) & AssignedInputs)
+					FString Name = ANSI_TO_TCHAR(Var->name);
+					Name.ReplaceInline(TEXT("."), TEXT("_"));
+					ReflectionOut.OutputVarNames.Add(Name);
+					CCHeaderWriter.WriteOutputAttribute(TEXT("out_Target"), *TypeQualifier, Var->location, /*bLocationPrefix:*/ true, /*bLocationSuffix:*/ true);
+				}
+				else
+				{
+					unsigned Location = Var->location;
+					unsigned SemanticIndex = Location;
+					check(Var->semantic);
+					unsigned i = (unsigned)strlen(Var->semantic);
+					check(i);
+					while (isdigit((unsigned char)(Var->semantic[i - 1])))
+					{
+						i--;
+					}
+					if (i < strlen(Var->semantic))
+					{
+						SemanticIndex = (unsigned)atoi(Var->semantic + i);
+						if (Location != SemanticIndex)
 						{
-							Location++;
+							Location = SemanticIndex;
 						}
+					}
 
-						if (Location != Var->location)
-						{
-							SPVRResult = Reflection.ChangeOutputVariableLocation(Var, Location);
-							check(SPVRResult == SPV_REFLECT_RESULT_SUCCESS);
-						}
-
-						uint32 ArrayCount = 1;
-						for (uint32 Dim = 0; Dim < Var->array.dims_count; Dim++)
-						{
-							ArrayCount *= Var->array.dims[Dim];
-						}
-
-						FString TypeQualifier;
-
-						auto const type = *Var->type_description;
-						uint32_t masked_type = type.type_flags & 0xF;
-
-						switch (masked_type) {
-						default: checkf(false, TEXT("unsupported component type %d"), masked_type); break;
-						case SPV_REFLECT_TYPE_FLAG_BOOL: TypeQualifier = TEXT("b"); break;
-						case SPV_REFLECT_TYPE_FLAG_INT: TypeQualifier = (type.traits.numeric.scalar.signedness ? TEXT("i") : TEXT("u")); break;
-						case SPV_REFLECT_TYPE_FLAG_FLOAT: TypeQualifier = (type.traits.numeric.scalar.width == 32 ? TEXT("f") : TEXT("h")); break;
-						}
-
-						if (type.type_flags & SPV_REFLECT_TYPE_FLAG_MATRIX)
-						{
-							TypeQualifier += FString::Printf(TEXT("%d%d"), type.traits.numeric.matrix.row_count, type.traits.numeric.matrix.column_count);
-						}
-						else if (type.type_flags & SPV_REFLECT_TYPE_FLAG_VECTOR)
-						{
-							TypeQualifier += FString::Printf(TEXT("%d"), type.traits.numeric.vector.component_count);
-						}
-						else
-						{
-							TypeQualifier += TEXT("1");
-						}
-
-						for (uint32 j = 0; j < ArrayCount; j++)
-						{
-							AssignedInputs |= (1 << (Location + j));
-						}
-
-						FString Name = ANSI_TO_TCHAR(Var->name);
-						Name.ReplaceInline(TEXT("."), TEXT("_"));
-						CCHeaderWriter.WriteOutputAttribute(*Name, *TypeQualifier, Location, /*bLocationPrefix:*/ true, /*bLocationSuffix:*/ false);
-					}
-				}
-			}
-		}
-
-		TArray<FString> InputVarNames;
-
-		{
-			uint32 AssignedInputs = 0;
-
-			ReflectionBindings.GatherInputAttributes(Reflection);
-			for (SpvReflectInterfaceVariable* Var : ReflectionBindings.InputAttributes)
-			{
-				if (Var->storage_class == SpvStorageClassInput && Var->built_in == -1)
-				{
-					unsigned Location = Var->location;
-					unsigned SemanticIndex = Location;
-					check(Var->semantic);
-					unsigned i = (unsigned)strlen(Var->semantic);
-					check(i);
-					while (isdigit((unsigned char)(Var->semantic[i - 1])))
-					{
-						i--;
-					}
-					if (i < strlen(Var->semantic))
-					{
-						SemanticIndex = (unsigned)atoi(Var->semantic + i);
-						if (Location != SemanticIndex)
-						{
-							Location = SemanticIndex;
-						}
-					}
-
-					while ((1 << Location) & AssignedInputs)
-					{
-						Location++;
-					}
+					while ((1 << Location) & AssignedInputs)
+					{
+						Location++;
+					}
 
 					if (Location != Var->location)
 					{
-						SPVRResult = Reflection.ChangeInputVariableLocation(Var, Location);
+						SPVRResult = Reflection.ChangeOutputVariableLocation(Var, Location);
 						check(SPVRResult == SPV_REFLECT_RESULT_SUCCESS);
 					}
 
@@ -2143,546 +1912,1042 @@ static bool CompileToGlslWithShaderConductor(
 
 					FString Name = ANSI_TO_TCHAR(Var->name);
 					Name.ReplaceInline(TEXT("."), TEXT("_"));
-					InputVarNames.Add(Name);
-					CCHeaderWriter.WriteInputAttribute(*Name, *TypeQualifier, Location, /*bLocationPrefix:*/ true, /*bLocationSuffix:*/ false);
+					CCHeaderWriter.WriteOutputAttribute(*Name, *TypeQualifier, Location, /*bLocationPrefix:*/ true, /*bLocationSuffix:*/ false);
 				}
 			}
 		}
+	}
 
-		// Overwrite updated SPIRV code
-		SpirvData = TArray<uint32>(Reflection.GetCode(), Reflection.GetCodeSize() / 4);
-
-		if (bDumpDebugInfo)
-		{
-			// SPIR-V file (Binary)
-			DumpDebugShaderBinary(Input, SpirvData.GetData(), SpirvData.Num() * sizeof(uint32), TEXT("spv"));
-			DumpDebugShaderDisassembledSpirv(Input, SpirvData.GetData(), SpirvData.Num() * sizeof(uint32), TEXT("spvasm"));
-		}
-
-		CrossCompiler::FShaderConductorTarget TargetDesc;
-
-		switch (Version)
-		{
-		case GLSL_150_ES3_1:	// ES3.1 Emulation
-			TargetDesc.CompileFlags.SetDefine(TEXT("force_flattened_io_blocks"), 1);
-			TargetDesc.CompileFlags.SetDefine(TEXT("emit_uniform_buffer_as_plain_uniforms"), 1);
-			TargetDesc.Language = CrossCompiler::EShaderConductorLanguage::Glsl;
-			TargetDesc.Version = 430;
-			break;
-		case GLSL_SWITCH_FORWARD:
-			TargetDesc.Language = CrossCompiler::EShaderConductorLanguage::Essl;
-			TargetDesc.Version = 320;
-			break;
-		case GLSL_SWITCH:
-			TargetDesc.Language = CrossCompiler::EShaderConductorLanguage::Glsl;
-			TargetDesc.Version = 430;
-			break;
-		case GLSL_ES3_1_ANDROID:
-		default:
-			TargetDesc.CompileFlags.SetDefine(TEXT("force_flattened_io_blocks"), 1);
-			TargetDesc.CompileFlags.SetDefine(TEXT("emit_uniform_buffer_as_plain_uniforms"), 1);
-			TargetDesc.CompileFlags.SetDefine(TEXT("force_temporary"), 1);
-
-			// If we have mobile multiview define set then set the view count and enable extension
-			const FString* MultiViewDefine = Input.Environment.GetDefinitions().Find(TEXT("MOBILE_MULTI_VIEW"));
-			if (Frequency == SF_Vertex && MultiViewDefine && *MultiViewDefine == "1")
-			{
-				TargetDesc.CompileFlags.SetDefine(TEXT("ovr_multiview_view_count"), 2);
-			}
-
-			TargetDesc.Language = CrossCompiler::EShaderConductorLanguage::Essl;
-			TargetDesc.Version = 320;
-			break;
-		}
+	{
+		uint32 AssignedInputs = 0;
 
-		TSet<FString> ExternalTextures;
-		int32 Pos = 0;
-#if !PLATFORM_MAC
-		TCHAR TextureExternalName[256];
-#else
-		ANSICHAR TextureExternalName[256];
-#endif
-		do
+		ReflectionBindings.GatherInputAttributes(Reflection);
+		for (SpvReflectInterfaceVariable* Var : ReflectionBindings.InputAttributes)
 		{
-			Pos = PreprocessedShader.Find(TEXT("TextureExternal"), ESearchCase::CaseSensitive, ESearchDir::FromStart, Pos + 15);
-			if (Pos != INDEX_NONE)
+			if (Var->storage_class == SpvStorageClassInput && Var->built_in == -1)
 			{
-#if PLATFORM_WINDOWS
-				if (swscanf_s(&PreprocessedShader[Pos], TEXT("TextureExternal %ls"), TextureExternalName, 256) == 1)
-#elif PLATFORM_MAC
-				if (sscanf(TCHAR_TO_ANSI(&PreprocessedShader[Pos]), "TextureExternal %s", TextureExternalName) == 1)
-#else // PLATFORM_LINUX
-				if (swscanf(TCHAR_TO_WCHAR(&PreprocessedShader[Pos]), L"TextureExternal %ls", TextureExternalName) == 1)
-#endif
+				unsigned Location = Var->location;
+				unsigned SemanticIndex = Location;
+				check(Var->semantic);
+				unsigned i = (unsigned)strlen(Var->semantic);
+				check(i);
+				while (isdigit((unsigned char)(Var->semantic[i - 1])))
 				{
-					FString Name = TextureExternalName;
-					if (Name.RemoveFromEnd(TEXT(";")))
+					i--;
+				}
+				if (i < strlen(Var->semantic))
+				{
+					SemanticIndex = (unsigned)atoi(Var->semantic + i);
+					if (Location != SemanticIndex)
 					{
-						ExternalTextures.Add(TEXT("SPIRV_Cross_Combined") + Name);
+						Location = SemanticIndex;
 					}
 				}
-			}
-		}
-		while (Pos != INDEX_NONE);
 
-		// Define type renaming callback after all external texture types have been gathered
-		TargetDesc.VariableTypeRenameCallback = [&ExternalTextures](const FAnsiStringView& VariableName, const FAnsiStringView& TypeName, FString& OutRenamedTypeName) -> bool
-		{
-			for (const FString& ExternalTex : ExternalTextures)
-			{
-				if (VariableName.Len() == ExternalTex.Len() && FCStringWide::Strncmp(ANSI_TO_TCHAR(VariableName.GetData()), *ExternalTex, ExternalTex.Len()) == 0)
+				while ((1 << Location) & AssignedInputs)
 				{
-					OutRenamedTypeName = TEXT("samplerExternalOES");
-					return true;
+					Location++;
 				}
-			}
-			return false;
-		};
-
-		std::string GlslSource;
-		const bool bGlslSourceCompileSucceeded = CompilerContext.CompileSpirvToSourceBuffer(
-			Options, TargetDesc, SpirvData.GetData(), SpirvData.Num() * sizeof(uint32),
-			[&GlslSource](const void* Data, uint32 Size)
-			{
-				GlslSource = std::string(reinterpret_cast<const ANSICHAR*>(Data), Size);
-			}
-		);
 
-		if (!bGlslSourceCompileSucceeded)
-		{
-			CompilerContext.FlushErrors(Output.Errors);
-			bCompilationFailed = true;
-		}
-		else
-		{
-			std::string LayoutString = "#extension ";
-			size_t LayoutPos = GlslSource.find(LayoutString);
-			if (LayoutPos != std::string::npos)
-			{
-				for (FString Name : InputVarNames)
+				if (Location != Var->location)
 				{
-					std::string DefineString = "#define ";
-					DefineString += TCHAR_TO_ANSI(*Name);
-					DefineString += " ";
-					DefineString += TCHAR_TO_ANSI(*Name.Replace(TEXT("in_var_"), TEXT("in_")));
-					DefineString += "\n";
-
-					GlslSource.insert(LayoutPos, DefineString);
+					SPVRResult = Reflection.ChangeInputVariableLocation(Var, Location);
+					check(SPVRResult == SPV_REFLECT_RESULT_SUCCESS);
 				}
-				for (FString Name : OutputVarNames)
-				{
-					std::string DefineString = "#define ";
-					DefineString += TCHAR_TO_ANSI(*Name);
-					DefineString += " ";
-					DefineString += TCHAR_TO_ANSI(*Name.Replace(TEXT("out_var_SV_"), TEXT("out_")));
-					DefineString += "\n";
 
-					GlslSource.insert(LayoutPos, DefineString);
-				}
-				/*if(!bEmulatedUBs)
+				uint32 ArrayCount = 1;
+				for (uint32 Dim = 0; Dim < Var->array.dims_count; Dim++)
 				{
-					for (auto const& Pair : UniformVarNames)
-					{
-						std::string DefineString = "#define ";
-						DefineString += Pair.first;
-						DefineString += " ";
-						DefineString += Pair.second;
-						DefineString += "\n";
-
-						GlslSource.insert(LayoutPos, DefineString);
-					}
-				}*/
-			}
+					ArrayCount *= Var->array.dims[Dim];
+				}
 
+				FString TypeQualifier;
 
-			// Perform FBF replacements
-			{
+				auto const type = *Var->type_description;
+				uint32_t masked_type = type.type_flags & 0xF;
 
-				size_t MainPos = GlslSource.find("#version 320 es");
+				switch (masked_type) {
+				default: checkf(false, TEXT("unsupported component type %d"), masked_type); break;
+				case SPV_REFLECT_TYPE_FLAG_BOOL: TypeQualifier = TEXT("b"); break;
+				case SPV_REFLECT_TYPE_FLAG_INT: TypeQualifier = (type.traits.numeric.scalar.signedness ? TEXT("i") : TEXT("u")); break;
+				case SPV_REFLECT_TYPE_FLAG_FLOAT: TypeQualifier = (type.traits.numeric.scalar.width == 32 ? TEXT("f") : TEXT("h")); break;
+				}
 
-				// Fallback if the shader is 310 
-				if (MainPos == std::string::npos)
+				if (type.type_flags & SPV_REFLECT_TYPE_FLAG_MATRIX)
 				{
-					MainPos = GlslSource.find("#version 310 es");
+					TypeQualifier += FString::Printf(TEXT("%d%d"), type.traits.numeric.matrix.row_count, type.traits.numeric.matrix.column_count);
 				}
-
-				if (MainPos != std::string::npos)
+				else if (type.type_flags & SPV_REFLECT_TYPE_FLAG_VECTOR)
 				{
-					MainPos += strlen("#version 320 es");
-					GlslSource.insert(MainPos, GlslFrameBufferDefines);
-					GlslSource.insert(MainPos, GlslFrameBufferExtensions);
+					TypeQualifier += FString::Printf(TEXT("%d"), type.traits.numeric.vector.component_count);
 				}
-
-				// Framebuffer Depth Fetch
+				else
 				{
-					std::string FBDFString = "_Globals.ARM_shader_framebuffer_fetch_depth_stencil";
-					std::string FBDFReplaceString = "_Globals_ARM_shader_framebuffer_fetch_depth_stencil";
-
-					size_t FramebufferDepthFetchPos = GlslSource.find(FBDFString);
-					bool UsesFramebufferDepthFetch = Frequency == SF_Pixel && GlslSource.find("_Globals.ARM_shader_framebuffer_fetch_depth_stencil") != std::string::npos;
-
-					if (UsesFramebufferDepthFetch)
-					{
-						GlslSource.erase(FramebufferDepthFetchPos, FBDFString.length());
-						GlslSource.insert(FramebufferDepthFetchPos, FBDFReplaceString);
-
-						std::string LastFragDepthARMString = "_Globals._RESERVED_IDENTIFIER_FIXUP_gl_LastFragDepthARM";
-						std::string LastFragDepthARMReplaceString = "gl_LastFragDepthARM";
-
-						size_t LastFragDepthARMStringPos = GlslSource.find(LastFragDepthARMString);
-
-						GlslSource.erase(LastFragDepthARMStringPos, LastFragDepthARMString.length());
-						GlslSource.insert(LastFragDepthARMStringPos, LastFragDepthARMReplaceString);
-					}
+					TypeQualifier += TEXT("1");
 				}
-			}
 
-			if (bEmulatedUBs)
-			{
+				for (uint32 j = 0; j < ArrayCount; j++)
 				{
-					bool bIsLayout = true;
+					AssignedInputs |= (1 << (Location + j));
+				}
 
-					std::string GlobalsSearchString = "layout\\(.*\\) uniform type_Globals$";
+				FString Name = ANSI_TO_TCHAR(Var->name);
+				Name.ReplaceInline(TEXT("."), TEXT("_"));
+				ReflectionOut.InputVarNames.Add(Name);
+				CCHeaderWriter.WriteInputAttribute(*Name, *TypeQualifier, Location, /*bLocationPrefix:*/ true, /*bLocationSuffix:*/ false);
+			}
+		}
+	}
+}
 
-					std::smatch RegexMatch;
-					std::regex_search(GlslSource, RegexMatch, std::regex(GlobalsSearchString));
-					
-					for (auto Match : RegexMatch)
-					{
-						GlobalsSearchString = Match;
-						break;
-					}
+static void ConvertToEmulatedUBs(std::string& GlslSource, const ReflectionData& ReflectData, const EShaderFrequency Frequency)
+{
+	const ANSICHAR* FrequencyPrefix = GetFrequencyPrefix(Frequency);
+	{
+		bool bIsLayout = true;
 
-					size_t GlobalPos = GlslSource.find(GlobalsSearchString);
+		std::string GlobalsSearchString = "layout\\(.*\\) uniform type_Globals$";
 
-					if (GlobalPos == std::string::npos)
-					{
-						GlobalsSearchString = "struct type_Globals";
-						GlobalPos = GlslSource.find(GlobalsSearchString);
-						bIsLayout = false;
-					}
+		std::smatch RegexMatch;
+		std::regex_search(GlslSource, RegexMatch, std::regex(GlobalsSearchString));
 
-					if (GlobalPos != std::string::npos)
-					{
-						std::string GlobalsEndSearchString;
-						if (bIsLayout)
-						{
-							GlobalsEndSearchString = "} _Globals;";
-						}
-						else
-						{
-							GlobalsEndSearchString = "uniform type_Globals _Globals;";
-						}
+		for (auto Match : RegexMatch)
+		{
+			GlobalsSearchString = Match;
+			break;
+		}
 
-						size_t GlobalEndPos = GlslSource.find(GlobalsEndSearchString);
+		size_t GlobalPos = GlslSource.find(GlobalsSearchString);
 
-						if (GlobalEndPos != std::string::npos)
-						{
-							GlslSource.erase(GlobalPos, GlobalEndPos - GlobalPos + GlobalsEndSearchString.length());
+		if (GlobalPos == std::string::npos)
+		{
+			GlobalsSearchString = "struct type_Globals";
+			GlobalPos = GlslSource.find(GlobalsSearchString);
+			bIsLayout = false;
+		}
 
-							bool UsesFramebufferFetch = Frequency == SF_Pixel && GlslSource.find("_Globals.ARM_shader_framebuffer_fetch") != std::string::npos;
+		if (GlobalPos != std::string::npos)
+		{
+			std::string GlobalsEndSearchString;
+			if (bIsLayout)
+			{
+				GlobalsEndSearchString = "} _Globals;";
+			}
+			else
+			{
+				GlobalsEndSearchString = "uniform type_Globals _Globals;";
+			}
 
-							std::string GlobalVarString = "_Globals.";
-							size_t GlobalVarPos = 0;
-							do
-							{
-								GlobalVarPos = GlslSource.find(GlobalVarString, GlobalVarPos);
-								if (GlobalVarPos != std::string::npos)
-								{
-									GlslSource.replace(GlobalVarPos, GlobalVarString.length(), "_Globals_");
-									for (std::string const& SearchString : GlobalArrays)
-									{
-										if (!GlslSource.compare(GlobalVarPos, SearchString.length(), SearchString))
-										{
-											GlslSource.replace(GlobalVarPos + SearchString.length(), 1, "(");
-
-											size_t ClosingBrace = GlslSource.find("]", GlobalVarPos + SearchString.length());
-											if (ClosingBrace != std::string::npos)
-												GlslSource.replace(ClosingBrace, 1, ")");
-										}
-									}
-								}
-							} while (GlobalVarPos != std::string::npos);
-
-							for (auto const& Pair : GlobalOffsets)
-							{
-								if (Pair.Value > 0)
-								{
-									std::string NewUniforms;
-									GetPackedUniformString(NewUniforms, FrequencyPrefix+std::string("u"), Pair.Key, Pair.Value);
-									GlslSource.insert(GlobalPos, NewUniforms);
-								}
-							}
+			size_t GlobalEndPos = GlslSource.find(GlobalsEndSearchString);
 
-							for (std::string const& Define : GlobalRemap)
-							{
-								GlslSource.insert(GlobalPos, Define);
-							}
+			if (GlobalEndPos != std::string::npos)
+			{
+				GlslSource.erase(GlobalPos, GlobalEndPos - GlobalPos + GlobalsEndSearchString.length());
 
-							if (UsesFramebufferFetch)
+				std::string GlobalVarString = "_Globals.";
+				size_t GlobalVarPos = 0;
+				do
+				{
+					GlobalVarPos = GlslSource.find(GlobalVarString, GlobalVarPos);
+					if (GlobalVarPos != std::string::npos)
+					{
+						GlslSource.replace(GlobalVarPos, GlobalVarString.length(), "_Globals_");
+						for (std::string const& SearchString : ReflectData.GlobalArrays)
+						{
+							if (!GlslSource.compare(GlobalVarPos, SearchString.length(), SearchString))
 							{
-								size_t MainPos = GlslSource.find("struct type_Globals");
-								if (MainPos != std::string::npos)
-									GlslSource.insert(MainPos, GlslFrameBufferDefines);
+								GlslSource.replace(GlobalVarPos + SearchString.length(), 1, "(");
 
-								size_t OutColor = GlslSource.find("0) out ");
-								if (OutColor != std::string::npos)
-									GlslSource.replace(OutColor, 7, "0) FRAME_BUFFERFETCH_STORAGE_QUALIFIER ");
+								size_t ClosingBrace = GlslSource.find("]", GlobalVarPos + SearchString.length());
+								if (ClosingBrace != std::string::npos)
+									GlslSource.replace(ClosingBrace, 1, ")");
 							}
 						}
 					}
-				}
+				} while (GlobalVarPos != std::string::npos);
 
-				for (const auto& PackedUBNamePair : PackedUBNames)
+				for (auto const& Pair : ReflectData.GlobalOffsets)
 				{
-					bool bIsLayout = true;
-					std::string UBSearchString = "layout\\(.*\\) uniform type_" + PackedUBNamePair.Value + "$";
-
-					std::smatch RegexMatch;
-					std::regex_search(GlslSource, RegexMatch, std::regex(UBSearchString));
-
-					for (auto Match : RegexMatch)
+					if (Pair.Value > 0)
 					{
-						UBSearchString = Match;
-						break;
+						std::string NewUniforms;
+						GetPackedUniformString(NewUniforms, FrequencyPrefix + std::string("u"), Pair.Key, Pair.Value);
+						GlslSource.insert(GlobalPos, NewUniforms);
 					}
+				}
 
-					size_t UBPos = GlslSource.find(UBSearchString);
+				for (std::string const& Define : ReflectData.GlobalRemap)
+				{
+					GlslSource.insert(GlobalPos, Define);
+				}
 
-					if (UBPos == std::string::npos)
-					{
-						UBSearchString = "struct type_" + PackedUBNamePair.Value;
-						UBPos = GlslSource.find(UBSearchString);
-						bIsLayout = false;
-					}
+				bool UsesFramebufferFetch = Frequency == SF_Pixel && GlslSource.find("_Globals.ARM_shader_framebuffer_fetch") != std::string::npos;
 
-					if (UBPos != std::string::npos)
-					{
-						std::string UBEndSearchString;
-						if (bIsLayout)
-						{
-							UBEndSearchString = "} " + PackedUBNamePair.Value + ";";
-						}
-						else
-						{
-							UBEndSearchString = "uniform type_" + PackedUBNamePair.Value + " " + PackedUBNamePair.Value + ";";
-						}
+				if (UsesFramebufferFetch)
+				{
+					size_t OutColor = GlslSource.find("0) out ");
+					if (OutColor != std::string::npos)
+						GlslSource.replace(OutColor, 7, "0) FRAME_BUFFERFETCH_STORAGE_QUALIFIER ");
+				}
+			}
+		}
+	}
 
-						size_t UBEndPos = GlslSource.find(UBEndSearchString);
+	for (const auto& PackedUBNamePair : ReflectData.PackedUBNames)
+	{
+		bool bIsLayout = true;
+		std::string UBSearchString = "layout\\(.*\\) uniform type_" + PackedUBNamePair.Value + "$";
 
-						if (UBEndPos != std::string::npos)
-						{
-							GlslSource.erase(UBPos, UBEndPos - UBPos + UBEndSearchString.length());
+		std::smatch RegexMatch;
+		std::regex_search(GlslSource, RegexMatch, std::regex(UBSearchString));
 
-							size_t UBVarPos = 0;
+		for (auto Match : RegexMatch)
+		{
+			UBSearchString = Match;
+			break;
+		}
 
-							for (const PackedUBMemberInfo& MemberInfo : PackedUBMemberInfos[PackedUBNamePair.Key])
-							{
-								std::string UBVarString = PackedUBNamePair.Value + "." + MemberInfo.Name;
-
-								UBVarPos = GlslSource.find(UBVarString);
-								while (UBVarPos != std::string::npos)
-								{
-									GlslSource.erase(UBVarPos, PackedUBNamePair.Value.length() + 1);
-									GlslSource.replace(UBVarPos, MemberInfo.Name.size(), MemberInfo.SanitizedName);
-
-									for (std::string const& SearchString : PackedUBArrays[PackedUBNamePair.Key])
-									{
-										if (!GlslSource.compare(UBVarPos, SearchString.length(), SearchString))
-										{
-											GlslSource.replace(UBVarPos + SearchString.length(), 1, "(");
-
-											size_t ClosingBrace = GlslSource.find("]", UBVarPos + SearchString.length());
-											if (ClosingBrace != std::string::npos)
-												GlslSource.replace(ClosingBrace, 1, ")");
-										}
-									}
-
-									UBVarPos = GlslSource.find(UBVarString);
-								}
-							}
+		size_t UBPos = GlslSource.find(UBSearchString);
 
-							for (auto const& Pair : PackedUBOffsets[PackedUBNamePair.Key])
-							{
-								if (Pair.Value > 0)
-								{
-									std::string NewUniforms;
-									std::string UniformPrefix = FrequencyPrefix + std::string("c") + std::to_string(PackedUBNamePair.Key);
-									GetPackedUniformString(NewUniforms, UniformPrefix, Pair.Key, Pair.Value);
-									GlslSource.insert(UBPos, NewUniforms);
-								}
-							}
+		if (UBPos == std::string::npos)
+		{
+			UBSearchString = "struct type_" + PackedUBNamePair.Value;
+			UBPos = GlslSource.find(UBSearchString);
+			bIsLayout = false;
+		}
 
-							for (std::string const& Define : PackedUBRemap[PackedUBNamePair.Key])
-							{
-								GlslSource.insert(UBPos, Define);
-							}
-						}
-					}
-				}
+		if (UBPos != std::string::npos)
+		{
+			std::string UBEndSearchString;
+			if (bIsLayout)
+			{
+				UBEndSearchString = "} " + PackedUBNamePair.Value + ";";
 			}
+			else
+			{
+				UBEndSearchString = "uniform type_" + PackedUBNamePair.Value + " " + PackedUBNamePair.Value + ";";
+			}
+
+			size_t UBEndPos = GlslSource.find(UBEndSearchString);
 
-			const size_t GlslSourceLen = GlslSource.length();
-			if (GlslSourceLen > 0)
+			if (UBEndPos != std::string::npos)
 			{
-				uint32 TextureIndex = 0;
-				for (const FString& Texture : Textures)
+				GlslSource.erase(UBPos, UBEndPos - UBPos + UBEndSearchString.length());
+
+				size_t UBVarPos = 0;
+
+				for (const PackedUBMemberInfo& MemberInfo : ReflectData.PackedUBMemberInfos[PackedUBNamePair.Key])
 				{
-					const size_t SamplerPos = GlslSource.find("\nuniform ");
+					std::string UBVarString = PackedUBNamePair.Value + "." + MemberInfo.Name;
 
-					TArray<FString> UsedSamplers;
-					FString SamplerString;
-					for (const FString& Sampler : Samplers)
+					UBVarPos = GlslSource.find(UBVarString);
+					while (UBVarPos != std::string::npos)
 					{
-						std::string SamplerName = "SPIRV_Cross_Combined";
-						SamplerName  += TCHAR_TO_ANSI(*(Texture + Sampler));
-						size_t FindCombinedSampler = GlslSource.find(SamplerName.c_str());
+						GlslSource.erase(UBVarPos, PackedUBNamePair.Value.length() + 1);
+						GlslSource.replace(UBVarPos, MemberInfo.Name.size(), MemberInfo.SanitizedName);
 
-						if (FindCombinedSampler != std::string::npos)
+						for (std::string const& SearchString : ReflectData.PackedUBArrays[PackedUBNamePair.Key])
 						{
-							checkf(SamplerPos != std::string::npos, TEXT("Generated GLSL shader is expected to have combined sampler '%s:%s' but no appropriate 'uniform' declaration could be found"), *Texture, *Sampler);
-
-							uint32 NewIndex = TextureIndex + UsedSamplers.Num();
-							std::string NewDefine = "#define ";
-							NewDefine += SamplerName;
-							NewDefine += " ";
-							NewDefine += FrequencyPrefix;
-							NewDefine += "s";
-							NewDefine += std::to_string(NewIndex);
-							NewDefine += "\n";
-							GlslSource.insert(SamplerPos+1, NewDefine);
-
-							// Do not add an entry for the dummy sampler as it will throw errors in the runtime checks
-							if (Sampler != SPIRV_DummySamplerName)
+							if (!GlslSource.compare(UBVarPos, SearchString.length(), SearchString))
 							{
-								UsedSamplers.Add(Sampler);
+								GlslSource.replace(UBVarPos + SearchString.length(), 1, "(");
+
+								size_t ClosingBrace = GlslSource.find("]", UBVarPos + SearchString.length());
+								if (ClosingBrace != std::string::npos)
+									GlslSource.replace(ClosingBrace, 1, ")");
 							}
-							
-							SamplerString += FString::Printf(TEXT("%s%s"), SamplerString.Len() ? TEXT(",") : TEXT(""), *Sampler);
 						}
-					}
 
-					// Rename texture buffers
-					std::string SamplerBufferName = std::string("samplerBuffer ") + TCHAR_TO_ANSI(*Texture);
-					size_t SamplerBufferPos = GlslSource.find(SamplerBufferName);
+						UBVarPos = GlslSource.find(UBVarString);
+					}
+				}
 
-					if (SamplerBufferPos != std::string::npos)
+				for (auto const& Pair : ReflectData.PackedUBOffsets[PackedUBNamePair.Key])
+				{
+					if (Pair.Value > 0)
 					{
-						std::string NewSamplerBufferName = std::string(FrequencyPrefix) + "s" + std::to_string(TextureIndex);
-
-						GlslSource.erase(SamplerBufferPos + SamplerBufferName.size() - Texture.Len(), Texture.Len());
-						GlslSource.insert(SamplerBufferPos + SamplerBufferName.size() - Texture.Len(), NewSamplerBufferName);
-
-						std::string SamplerTexFetchName = std::string("texelFetch(") + TCHAR_TO_ANSI(*Texture);
-						size_t SamplerTexFetchPos = GlslSource.find(SamplerTexFetchName);
+						std::string NewUniforms;
+						std::string UniformPrefix = FrequencyPrefix + std::string("c") + std::to_string(PackedUBNamePair.Key);
+						GetPackedUniformString(NewUniforms, UniformPrefix, Pair.Key, Pair.Value);
+						GlslSource.insert(UBPos, NewUniforms);
+					}
+				}
 
-						while (SamplerTexFetchPos != std::string::npos)
-						{
+				for (std::string const& Define : ReflectData.PackedUBRemap[PackedUBNamePair.Key])
+				{
+					GlslSource.insert(UBPos, Define);
+				}
+			}
+		}
+	}
+}
 
-							GlslSource.erase(SamplerTexFetchPos + SamplerTexFetchName.size() - Texture.Len(), Texture.Len());
-							GlslSource.insert(SamplerTexFetchPos + SamplerTexFetchName.size() - Texture.Len(), NewSamplerBufferName);
+const ANSICHAR* GlslFrameBufferExtensions =
+"\n\n#ifdef GL_ARM_shader_framebuffer_fetch_depth_stencil\n"
+"\t#extension GL_ARM_shader_framebuffer_fetch_depth_stencil : enable\n"
+"#elif defined(GL_EXT_shader_framebuffer_fetch)\n"
+"\t#extension GL_EXT_shader_framebuffer_fetch : enable\n"
+"\t#define FBF_STORAGE_QUALIFIER inout\n"
+"#endif\n"
+"#extension GL_EXT_texture_buffer : enable\n"
+"// end extensions";
+
+// GLSL framebuffer macro definitions. Used to patch GLSL output source.
+const ANSICHAR* GlslFrameBufferDefines =
+"\n\n#ifdef UE_EXT_shader_framebuffer_fetch\n"
+"#define _Globals_ARM_shader_framebuffer_fetch 0\n"
+"#define FRAME_BUFFERFETCH_STORAGE_QUALIFIER inout\n"
+"#define _Globals_gl_FragColor out_var_SV_Target0\n"
+"#define _Globals_gl_LastFragColorARM vec4(0.0, 0.0, 0.0, 0.0)\n"
+"#elif defined( GL_ARM_shader_framebuffer_fetch)\n"
+"#define _Globals_ARM_shader_framebuffer_fetch 1\n"
+"#define FRAME_BUFFERFETCH_STORAGE_QUALIFIER out\n"
+"#define _Globals_gl_FragColor vec4(0.0, 0.0, 0.0, 0.0)\n"
+"#define _Globals_gl_LastFragColorARM gl_LastFragDepthARM\n"
+"#else\n"
+"#define FRAME_BUFFERFETCH_STORAGE_QUALIFIER out\n"
+"#define _Globals_ARM_shader_framebuffer_fetch 0\n"
+"#define _Globals_gl_FragColor vec4(0.0, 0.0, 0.0, 0.0)\n"
+"#define _Globals_gl_LastFragColorARM vec4(0.0, 0.0, 0.0, 0.0)\n"
+"#endif\n"
+"#ifdef GL_ARM_shader_framebuffer_fetch_depth_stencil\n"
+"#define _Globals_ARM_shader_framebuffer_fetch_depth_stencil 1u\n"
+"#else\n"
+"#define _Globals_ARM_shader_framebuffer_fetch_depth_stencil 0u\n"
+"#endif\n";
+
+struct GLSLCompileParameters
+{
+	CrossCompiler::FShaderConductorContext* CompilerContext;
+	CrossCompiler::FHlslccHeaderWriter* CCHeaderWriter;
+	CrossCompiler::FShaderConductorTarget* TargetDesc;
+	CrossCompiler::FShaderConductorOptions* Options;
 
-							SamplerTexFetchPos = GlslSource.find(SamplerTexFetchName);
-						}
-					}
+	FShaderCompilerOutput* Output;
+	TArray<uint32>* SpirvData;
 
-					const uint32 SamplerCount = FMath::Max(1, UsedSamplers.Num());
-					CCHeaderWriter.WriteSRV(*Texture, TextureIndex, SamplerCount, UsedSamplers);
-					TextureIndex += SamplerCount;
-				}
+	EShaderFrequency Frequency;
+	const ANSICHAR* SPIRV_DummySamplerName;
+};
 
-				// UAVS, rename as ci0 format
-				uint32_t UAVIndex = 0;
-				for (const std::string& UAV : UAVs)
-				{
-					std::string NewUAVName = FrequencyPrefix + std::string("i") + std::to_string(UAVIndex);
+bool GenerateGlslShader(std::string& OutString, GLSLCompileParameters& GLSLCompileParams, ReflectionData& ReflectData, bool bWriteToCCHeader, bool bIsDeferred, bool bEmulatedUBs)
+{
+	const bool bGlslSourceCompileSucceeded = GLSLCompileParams.CompilerContext->CompileSpirvToSourceBuffer(
+		*GLSLCompileParams.Options, *GLSLCompileParams.TargetDesc, GLSLCompileParams.SpirvData->GetData(), GLSLCompileParams.SpirvData->Num() * sizeof(uint32),
+		[&OutString](const void* Data, uint32 Size)
+		{
+			OutString = std::string(reinterpret_cast<const ANSICHAR*>(Data), Size);
+		}
+	);
+
+	if (!bGlslSourceCompileSucceeded)
+	{
+		GLSLCompileParams.CompilerContext->FlushErrors(GLSLCompileParams.Output->Errors);
+		return false;
+	}
+
+	const ANSICHAR* FrequencyPrefix = GetFrequencyPrefix(GLSLCompileParams.Frequency);
+
+	std::string LayoutString = "#extension ";
+	size_t LayoutPos = OutString.find(LayoutString);
+	if (LayoutPos != std::string::npos)
+	{
+		for (FString Name : ReflectData.InputVarNames)
+		{
+			std::string DefineString = "#define ";
+			DefineString += TCHAR_TO_ANSI(*Name);
+			DefineString += " ";
+			DefineString += TCHAR_TO_ANSI(*Name.Replace(TEXT("in_var_"), TEXT("in_")));
+			DefineString += "\n";
+
+			OutString.insert(LayoutPos, DefineString);
+		}
+		for (FString Name : ReflectData.OutputVarNames)
+		{
+			std::string DefineString = "#define ";
+			DefineString += TCHAR_TO_ANSI(*Name);
+			DefineString += " ";
+			DefineString += TCHAR_TO_ANSI(*Name.Replace(TEXT("out_var_SV_"), TEXT("out_")));
+			DefineString += "\n";
+
+			OutString.insert(LayoutPos, DefineString);
+		}
+	}
+
+	// Perform FBF replacements
+	{
+
+		size_t MainPos = OutString.find("#version 320 es");
+
+		// Fallback if the shader is 310 
+		if (MainPos == std::string::npos)
+		{
+			MainPos = OutString.find("#version 310 es");
+		}
+
+		if (MainPos != std::string::npos)
+		{
+			MainPos += strlen("#version 320 es");
+		}
+
+		// Framebuffer Depth Fetch
+		{
+			std::string FBDFString = "_Globals.ARM_shader_framebuffer_fetch_depth_stencil";
+			std::string FBDFReplaceString = "_Globals_ARM_shader_framebuffer_fetch_depth_stencil";
 
-					// Find instances of UAVs
-					std::string RegexExpression = "(?:^|\\W|\\S)" + UAV + "(?:$|\\W)";
+			std::string FBFString = "_Globals.ARM_shader_framebuffer_fetch";
+			std::string FBFReplaceString = "_Globals_ARM_shader_framebuffer_fetch";
 
-					std::cmatch RegexMatch;
-					std::vector<std::string> RegexMatches;
+			bool UsesFramebufferDepthFetch = GLSLCompileParams.Frequency == SF_Pixel && OutString.find(FBDFString) != std::string::npos;
+			bool UsesFramebufferFetch = GLSLCompileParams.Frequency == SF_Pixel && OutString.find(FBFString) != std::string::npos;
 
-					const char * TempString = GlslSource.c_str();
-					while (std::regex_search(TempString, RegexMatch, std::regex(RegexExpression)))
+			if (UsesFramebufferDepthFetch || UsesFramebufferFetch)
+			{
+				OutString.insert(MainPos, GlslFrameBufferDefines);
+				OutString.insert(MainPos, GlslFrameBufferExtensions);
+
+				size_t FramebufferDepthFetchPos = OutString.find(FBDFString);
+
+				OutString.erase(FramebufferDepthFetchPos, FBDFString.length());
+				OutString.insert(FramebufferDepthFetchPos, FBDFReplaceString);
+
+				std::string LastFragDepthARMString = "_Globals._RESERVED_IDENTIFIER_FIXUP_gl_LastFragDepthARM";
+				std::string LastFragDepthARMReplaceString = "gl_LastFragDepthARM";
+
+				size_t LastFragDepthARMStringPos = OutString.find(LastFragDepthARMString);
+
+				OutString.erase(LastFragDepthARMStringPos, LastFragDepthARMString.length());
+				OutString.insert(LastFragDepthARMStringPos, LastFragDepthARMReplaceString);
+
+				MainPos = OutString.find("void main()");
+
+				// Add support for framebuffer fetch depth when ARM extension is not supported
+				if (MainPos != std::string::npos)
+				{
+					std::string DepthBufferIndex = bIsDeferred ? "4" : "1";
+					std::string DepthBufferOutVarString = "out_var_SV_Target" + DepthBufferIndex;
+
+					OutString.insert(MainPos, "float GLFetchDepthBuffer()\n"
+						"{\n"
+						"\t#ifdef GL_ARM_shader_framebuffer_fetch_depth_stencil\n"
+						"\treturn gl_LastFragDepthARM;\n"
+						"\t#else\n"
+						"\treturn " + DepthBufferOutVarString + ".x;\n"
+						"\t#endif\n"
+						"}\n");
+
+					size_t DepthBufferOutVarPos = OutString.find(DepthBufferOutVarString + ";");
+					if (DepthBufferOutVarPos == std::string::npos)
 					{
-						RegexMatches.push_back(RegexMatch.str());
-						TempString = RegexMatch.suffix().first;
+						OutString.insert(MainPos, "\n#ifndef GL_ARM_shader_framebuffer_fetch_depth_stencil\n"
+							"layout(location = " + DepthBufferIndex + ") inout highp vec4 " + DepthBufferOutVarString + "; \n"
+							"#endif\n");
 					}
-
-					for(const auto & Match : RegexMatches)
+					else
 					{
-						size_t MatchOffset = Match.find_first_of(UAV);
-						size_t MatchPos = GlslSource.find(Match);
-
-						GlslSource.replace(MatchPos + MatchOffset, UAV.size(), NewUAVName);
+						size_t QualifierPos = OutString.find_last_of("out", DepthBufferOutVarPos - 1);
+						OutString.erase(QualifierPos - 2, 3);
+						OutString.insert(QualifierPos - 2, "inout");
 					}
 
-					UAVIndex++;
+					std::string DepthFetchReplacePosString = "0.0 : gl_LastFragDepthARM";
+					size_t DepthFetchReplacePos = OutString.find(DepthFetchReplacePosString);
+
+					while (DepthFetchReplacePos != std::string::npos)
+					{
+						OutString.erase(DepthFetchReplacePos, DepthFetchReplacePosString.length());
+						OutString.insert(DepthFetchReplacePos, "GLFetchDepthBuffer() : GLFetchDepthBuffer()");
+						DepthFetchReplacePos = OutString.find(DepthFetchReplacePosString);
+					}
 				}
+			}
+		}
+	}
+
+	if (bEmulatedUBs)
+	{
+		ConvertToEmulatedUBs(OutString, ReflectData, GLSLCompileParams.Frequency);
+	}
+
+	const size_t GlslSourceLen = OutString.length();
+	if (GlslSourceLen > 0)
+	{
+		uint32 TextureIndex = 0;
+		for (const FString& Texture : ReflectData.Textures)
+		{
+			const size_t SamplerPos = OutString.find("\nuniform ");
+
+			TArray<FString> UsedSamplers;
+			FString SamplerString;
+			for (const FString& Sampler : ReflectData.Samplers)
+			{
+				std::string SamplerName = "SPIRV_Cross_Combined";
+				SamplerName += TCHAR_TO_ANSI(*(Texture + Sampler));
+				size_t FindCombinedSampler = OutString.find(SamplerName.c_str());
 
-				// StructuredBuffers, rename as ci0 format
-				for (const std::string& SBuffer : StructuredBuffers)
+				if (FindCombinedSampler != std::string::npos)
 				{
-					std::string NewSBufferName = FrequencyPrefix + std::string("i") + std::to_string(UAVIndex);
-					std::string NewSBufferVarName = NewSBufferName + std::string("_VAR");
-					
-					std::string SearchString = "} " + SBuffer;
-					size_t SBufferPos = GlslSource.find(SearchString);
-
-					size_t SBufferEndPos = GlslSource.find(";", SBufferPos);
-					
-					std::string ReplacementSubStr = GlslSource.substr(SBufferPos + 2, SBufferEndPos - SBufferPos - 2);
-					GlslSource.erase(SBufferPos + 1, SBufferEndPos - SBufferPos - 1);
-
-					const std::string SBufferData = "_m0";
-					size_t SBufferDataPos = GlslSource.rfind(SBufferData, SBufferPos);
-					GlslSource.replace(SBufferDataPos, SBufferData.size(), NewSBufferName);
-
-					const std::string BufferString = " buffer ";
-					size_t BufferNamePos = GlslSource.rfind(BufferString, SBufferPos);
-					size_t BufferLineEndPos = GlslSource.find("\n", BufferNamePos);
-					size_t ReplacePos = BufferNamePos + BufferString.size();
-
-					GlslSource.replace(ReplacePos, BufferLineEndPos - ReplacePos, NewSBufferVarName);
-
-					// Replace the usage of StructuredBuffer with new name
-					size_t CurPos = GlslSource.find(ReplacementSubStr + ".");
-					while (CurPos != std::string::npos)
+					checkf(SamplerPos != std::string::npos, TEXT("Generated GLSL shader is expected to have combined sampler '%s:%s' but no appropriate 'uniform' declaration could be found"), *Texture, *Sampler);
+
+					uint32 NewIndex = TextureIndex + UsedSamplers.Num();
+					std::string NewDefine = "#define ";
+					NewDefine += SamplerName;
+					NewDefine += " ";
+					NewDefine += FrequencyPrefix;
+					NewDefine += "s";
+					NewDefine += std::to_string(NewIndex);
+					NewDefine += "\n";
+					OutString.insert(SamplerPos + 1, NewDefine);
+
+					// Do not add an entry for the dummy sampler as it will throw errors in the runtime checks
+					if (Sampler != GLSLCompileParams.SPIRV_DummySamplerName)
 					{
-						// Offset by 4 to account for ._m0
-						GlslSource.replace(CurPos, ReplacementSubStr.size()+4, NewSBufferName);
-						CurPos = GlslSource.find(ReplacementSubStr + ".");
+						UsedSamplers.Add(Sampler);
 					}
-					
-					UAVIndex++;
+
+					SamplerString += FString::Printf(TEXT("%s%s"), SamplerString.Len() ? TEXT(",") : TEXT(""), *Sampler);
+				}
+			}
+
+			// Rename texture buffers
+			std::string SamplerBufferName = std::string("samplerBuffer ") + TCHAR_TO_ANSI(*Texture);
+			size_t SamplerBufferPos = OutString.find(SamplerBufferName);
+
+			if (SamplerBufferPos != std::string::npos)
+			{
+				std::string NewSamplerBufferName = std::string(FrequencyPrefix) + "s" + std::to_string(TextureIndex);
+
+				OutString.erase(SamplerBufferPos + SamplerBufferName.size() - Texture.Len(), Texture.Len());
+				OutString.insert(SamplerBufferPos + SamplerBufferName.size() - Texture.Len(), NewSamplerBufferName);
+
+				std::string SamplerTexFetchName = std::string("texelFetch(") + TCHAR_TO_ANSI(*Texture);
+				size_t SamplerTexFetchPos = OutString.find(SamplerTexFetchName);
+
+				while (SamplerTexFetchPos != std::string::npos)
+				{
+
+					OutString.erase(SamplerTexFetchPos + SamplerTexFetchName.size() - Texture.Len(), Texture.Len());
+					OutString.insert(SamplerTexFetchPos + SamplerTexFetchName.size() - Texture.Len(), NewSamplerBufferName);
+
+					SamplerTexFetchPos = OutString.find(SamplerTexFetchName);
 				}
+			}
+
+			const uint32 SamplerCount = FMath::Max(1, UsedSamplers.Num());
+			if (bWriteToCCHeader)
+			{
+				GLSLCompileParams.CCHeaderWriter->WriteSRV(*Texture, TextureIndex, SamplerCount, UsedSamplers);
+			}
+			TextureIndex += SamplerCount;
+		}
+
+		// UAVS, rename as ci0 format
+		uint32_t UAVIndex = 0;
+		for (const std::string& UAV : ReflectData.UAVs)
+		{
+			std::string NewUAVName = FrequencyPrefix + std::string("i") + std::to_string(UAVIndex);
+
+			// Find instances of UAVs
+			std::string RegexExpression = "(?:^|\\W|\\S)" + UAV + "(?:$|\\W)";
+
+			std::cmatch RegexMatch;
+			std::vector<std::string> RegexMatches;
+
+			const char* TempString = OutString.c_str();
+			while (std::regex_search(TempString, RegexMatch, std::regex(RegexExpression)))
+			{
+				RegexMatches.push_back(RegexMatch.str());
+				TempString = RegexMatch.suffix().first;
+			}
+
+			for (const auto& Match : RegexMatches)
+			{
+				size_t MatchOffset = Match.find_first_of(UAV);
+				size_t MatchPos = OutString.find(Match);
+
+				OutString.replace(MatchPos + MatchOffset, UAV.size(), NewUAVName);
+			}
+
+			UAVIndex++;
+		}
+
+		// StructuredBuffers, rename as ci0 format
+		for (const std::string& SBuffer : ReflectData.StructuredBuffers)
+		{
+			std::string NewSBufferName = FrequencyPrefix + std::string("i") + std::to_string(UAVIndex);
+			std::string NewSBufferVarName = NewSBufferName + std::string("_VAR");
+
+			std::string SearchString = "} " + SBuffer;
+			size_t SBufferPos = OutString.find(SearchString);
+
+			size_t SBufferEndPos = OutString.find(";", SBufferPos);
+
+			std::string ReplacementSubStr = OutString.substr(SBufferPos + 2, SBufferEndPos - SBufferPos - 2);
+			OutString.erase(SBufferPos + 1, SBufferEndPos - SBufferPos - 1);
+
+			const std::string SBufferData = "_m0";
+			size_t SBufferDataPos = OutString.rfind(SBufferData, SBufferPos);
+			OutString.replace(SBufferDataPos, SBufferData.size(), NewSBufferName);
+
+			const std::string BufferString = " buffer ";
+			size_t BufferNamePos = OutString.rfind(BufferString, SBufferPos);
+			size_t BufferLineEndPos = OutString.find("\n", BufferNamePos);
+			size_t ReplacePos = BufferNamePos + BufferString.size();
+
+			OutString.replace(ReplacePos, BufferLineEndPos - ReplacePos, NewSBufferVarName);
 
-				for (const auto & pair : UniformVarNames)
+			// Replace the usage of StructuredBuffer with new name
+			size_t CurPos = OutString.find(ReplacementSubStr + ".");
+			while (CurPos != std::string::npos)
+			{
+				// Offset by 4 to account for ._m0
+				OutString.replace(CurPos, ReplacementSubStr.size() + 4, NewSBufferName);
+				CurPos = OutString.find(ReplacementSubStr + ".");
+			}
+
+			UAVIndex++;
+		}
+
+		for (const auto& pair : ReflectData.UniformVarNames)
+		{
+			std::string OldUniformTypeName = "uniform type_" + pair.first;
+			std::string NewUniformTypeName = "uniform " + pair.second;
+
+			size_t UniformTypePos = OutString.find(OldUniformTypeName);
+
+			if (UniformTypePos != std::string::npos)
+			{
+				OutString.erase(UniformTypePos, OldUniformTypeName.length());
+				OutString.insert(UniformTypePos, NewUniformTypeName);
+			}
+		}
+	}
+
+	return true;
+}
+
+bool GenerateDeferredMobileShaders(std::string& GlslSource, GLSLCompileParameters& GLSLCompileParams, const std::string& HlslString, ReflectionData& ReflectData, bool bWriteToCCHeader, bool bIsDeferred, bool bEmulatedUBs)
+{
+	bool bHasGbufferTextures[4];
+	bHasGbufferTextures[0] = HlslString.find("GENERATED_SubpassFetchAttachment0") != std::string::npos;
+	bHasGbufferTextures[1] = HlslString.find("GENERATED_SubpassFetchAttachment1") != std::string::npos;
+	bHasGbufferTextures[2] = HlslString.find("GENERATED_SubpassFetchAttachment2") != std::string::npos;
+	bHasGbufferTextures[3] = HlslString.find("GENERATED_SubpassFetchAttachment3") != std::string::npos;
+
+	bool bHasGBufferOutputs[4];
+	bHasGBufferOutputs[0] = HlslString.find("OutProxy") != std::string::npos;
+	bHasGBufferOutputs[1] = HlslString.find("OutGBufferA") != std::string::npos;
+	bHasGBufferOutputs[2] = HlslString.find("OutGBufferB") != std::string::npos;
+	bHasGBufferOutputs[3] = HlslString.find("OutGBufferC") != std::string::npos;
+
+	bool bHasInputs = false;
+	bool bHasOutputs = false;
+
+	for (uint32_t i = 0; i < 4; ++i)
+	{
+		bHasInputs |= bHasGbufferTextures[i];
+		bHasOutputs |= bHasGBufferOutputs[i];
+	}
+
+	bool bHasInOut = (bHasInputs && bHasOutputs) || HlslString.find("OutTarget0") != std::string::npos;
+
+	if (bHasInputs || bHasOutputs)
+	{
+		std::string OutString;
+		std::string ESVersionString = "#version 320 es";
+		OutString += ESVersionString + "\n";
+		OutString += "#ifdef UE_MRT_FRAMEBUFFER_FETCH\n";
+
+		FShaderCompilerDefinitions CompileFlagsCopy = GLSLCompileParams.TargetDesc->CompileFlags;
+
+		// Add defines for FBF in spirv-glsl
+		for (uint32_t i = 1; i < 4; ++i)
+		{
+			if (bHasGbufferTextures[i])
+			{
+				FString DefineKey = FString::Printf(TEXT("remap_ext_framebuffer_fetch%d"), i);
+				FString DefineValue = FString::Printf(TEXT("%d %d"), i + 1, i);
+				GLSLCompileParams.TargetDesc->CompileFlags.SetDefine(*DefineKey, *DefineValue);
+			}
+		}
+
+		std::string FBFSourceString;
+		if (!GenerateGlslShader(FBFSourceString, GLSLCompileParams, ReflectData, true, true, bEmulatedUBs))
+		{
+			return false;
+		}
+
+		size_t VersionStringPos = FBFSourceString.find(ESVersionString);
+		if (VersionStringPos != std::string::npos)
+		{
+			FBFSourceString.replace(VersionStringPos, ESVersionString.length(), "");
+		}
+
+		GLSLCompileParams.TargetDesc->CompileFlags = CompileFlagsCopy;
+
+		OutString += FBFSourceString;
+
+		// Generate PLS shader
+		OutString += "#else\n";
+
+		// Using rgb10_a2ui because r11_g11_b10 has issues with swizzle from 4 to 3 components in DXC
+		static const FString GBufferOutputNames[] =
+		{
+			TEXT("rgb10_a2 out.var.SV_Target0"),
+			TEXT("rgba8 out.var.SV_Target1"),
+			TEXT("rgba8 out.var.SV_Target2"),
+			TEXT("rgba8 out.var.SV_Target3"),
+		};
+
+		static const FString GBufferInputNames[] =
+		{
+			TEXT("rgb10_a2 GENERATED_SubpassFetchAttachment0"),
+			TEXT("rgba8 GENERATED_SubpassFetchAttachment1"),
+			TEXT("rgba8 GENERATED_SubpassFetchAttachment2"),
+			TEXT("rgba8 GENERATED_SubpassFetchAttachment3"),
+		};
+
+		static const FString GBufferInOutNames[] =
+		{
+			TEXT("rgb10_a2 GENERATED_SubpassFetchAttachment0 out.var.SV_Target0"),
+			TEXT("rgba8 GENERATED_SubpassFetchAttachment1 out.var.SV_Target1"),
+			TEXT("rgba8 GENERATED_SubpassFetchAttachment2 out.var.SV_Target2"),
+			TEXT("rgba8 GENERATED_SubpassFetchAttachment3 out.var.SV_Target3"),
+		};
+
+		// Add defines for PLS in spirv-glsl
+		{
+			FString DefineKey = "";
+			FString DefineValue = "";
+
+			for (uint32_t i = 0; i < 4; ++i)
+			{
+				if (bHasInOut)
+				{
+					DefineKey = FString::Printf(TEXT("pls_io%d"), i);
+					DefineValue = GBufferInOutNames[i];
+				}
+				else if (bHasInputs)
 				{
-					std::string OldUniformTypeName = "uniform type_" + pair.first;
-					std::string NewUniformTypeName = "uniform " + pair.second;
+					DefineKey = FString::Printf(TEXT("pls_in%d"), i);
+					DefineValue = GBufferInputNames[i];
+				}
+				else
+				{
+					DefineKey = FString::Printf(TEXT("pls_out%d"), i);
+					DefineValue = GBufferOutputNames[i];
+				}
+
+				GLSLCompileParams.TargetDesc->CompileFlags.SetDefine(*DefineKey, *DefineValue);
+			}
+		}
+
+		std::string PLSSourceString;
+		if (!GenerateGlslShader(PLSSourceString, GLSLCompileParams, ReflectData, false, true, bEmulatedUBs))
+		{
+			return false;
+		}
+
+		// Replace assignment of output to buffer 0 to additive if the output proxy additive is used
+		if (HlslString.find("OutProxyAdditive") != std::string::npos ||
+			HlslString.find("OutTarget0") != std::string::npos)
+		{
+			std::string OutProxyString = "GENERATED_SubpassFetchAttachment0 =";
+			std::string OutProxyModifiedString = "GENERATED_SubpassFetchAttachment0 +=";
+
+			size_t OutProxyStringPos = PLSSourceString.find(OutProxyString);
+			while (OutProxyStringPos != std::string::npos)
+			{
+				PLSSourceString.replace(OutProxyStringPos, OutProxyString.size(), OutProxyModifiedString);
+				OutProxyStringPos = PLSSourceString.find(OutProxyString);
+			}
+		}
+
+		// Strip version string
+		VersionStringPos = PLSSourceString.find(ESVersionString);
+		if (VersionStringPos != std::string::npos)
+		{
+			PLSSourceString.replace(VersionStringPos, ESVersionString.length(), "", 0);
+		}
+
+		OutString += PLSSourceString;
+
+		OutString += "#endif\n";
+		GlslSource = OutString;
+	}
+	else
+	{
+		if (!GenerateGlslShader(GlslSource, GLSLCompileParams, ReflectData, true, false, bEmulatedUBs))
+		{
+			return false;
+		}
+	}
+
+	return true;
+}
+
+static bool CompileToGlslWithShaderConductor(
+	const FShaderCompilerInput&	Input,
+	FShaderCompilerOutput&		Output,
+	const FString&				WorkingDirectory,
+	GLSLVersion					Version,
+	const EShaderFrequency		Frequency,
+	uint32						CCFlags,
+	const FString&				PreprocessedShader,
+	char*&						OutGlslShaderSource)
+{
+	CrossCompiler::FShaderConductorContext CompilerContext;
 
-					size_t UniformTypePos = GlslSource.find(OldUniformTypeName);
+	const bool bDumpDebugInfo = (Input.DumpDebugInfoPath != TEXT("") && IFileManager::Get().DirectoryExists(*Input.DumpDebugInfoPath));
+	const bool bRewriteHlslSource = true;
 
-					if (UniformTypePos != std::string::npos)
+	// Initialize compilation options for ShaderConductor
+	CrossCompiler::FShaderConductorOptions Options;
+	Options.bDisableScalarBlockLayout = true;
+	Options.bRemapAttributeLocations = true;
+	Options.bPreserveStorageInput = true;
+	
+	// Convert input strings from FString to ANSI strings
+	std::string SourceData(TCHAR_TO_UTF8(*PreprocessedShader));
+	std::string FileName(TCHAR_TO_UTF8(*Input.VirtualSourceFilePath));
+	std::string EntryPointName(TCHAR_TO_UTF8(*Input.EntryPointName));
+	
+	bool bEmulatedUBs = Input.Environment.CompilerFlags.Contains(CFLAG_UseEmulatedUB);
+
+	// HLSL framebuffer declarations. Used to modify HLSL input source.
+	const ANSICHAR* HlslFrameBufferDeclarations =
+		"float4 gl_FragColor;\n"
+		"float4 gl_LastFragColorARM;\n"
+		"float gl_LastFragDepthARM;\n"
+		"bool ARM_shader_framebuffer_fetch;\n"
+		"bool ARM_shader_framebuffer_fetch_depth_stencil;\n"
+		"float4 FramebufferFetchES2()\n"
+		"{\n"
+		"  if (!ARM_shader_framebuffer_fetch)\n"
+		"  {\n"
+		"    return gl_FragColor;\n"
+		"  }\n"
+		"  else\n"
+		"  {\n"
+		"    return gl_LastFragColorARM;\n"
+		"  }\n"
+		"}\n"
+		"float DepthbufferFetchES2()\n"
+		"{\n"
+		"  return (ARM_shader_framebuffer_fetch_depth_stencil == 0 ? 0.0 : gl_LastFragDepthARM);\n"
+		"}\n"
+		;
+
+	if (SourceData.find("DepthbufferFetchES2") != std::string::npos ||
+		SourceData.find("FramebufferFetchES2") != std::string::npos)
+	{
+		SourceData = HlslFrameBufferDeclarations + SourceData;
+	}
+	
+	// Inject additional macro definitions to circumvent missing features: external textures
+	FShaderCompilerDefinitions AdditionalDefines;
+	AdditionalDefines.SetDefine(TEXT("TextureExternal"), TEXT("Texture2D"));
+
+	if (bDumpDebugInfo)
+	{
+		FString DirectCompileLine = CrossCompiler::CreateResourceTableFromEnvironment(Input.Environment);
+
+		DirectCompileLine += TEXT("#if 0 /*DIRECT COMPILE*/\n");
+		DirectCompileLine += CreateShaderCompilerWorkerDirectCommandLine(Input, CCFlags);
+		DirectCompileLine += TEXT("\n#endif /*DIRECT COMPILE*/\n");
+
+		DumpDebugUSF(Input, (PreprocessedShader + DirectCompileLine), CCFlags);
+
+		if (Input.bGenerateDirectCompileFile)
+		{
+			FFileHelper::SaveStringToFile(CreateShaderCompilerWorkerDirectCommandLine(Input), *(Input.DumpDebugInfoPath / TEXT("DirectCompile.txt")));
+		}
+	}
+
+	// Load shader source into compiler context
+	CompilerContext.LoadSource(SourceData.c_str(), FileName.c_str(), EntryPointName.c_str(), Frequency, &AdditionalDefines);
+
+	bool bCompilationFailed = false;
+
+	if (bRewriteHlslSource)
+	{
+		// Rewrite HLSL source code to remove unused global resources and variables
+		Options.bRemoveUnusedGlobals = true;
+		if (CompilerContext.RewriteHlsl(Options))
+		{
+			// Adopt new rewritten shader source
+			SourceData = CompilerContext.GetSourceString();
+
+			if (bDumpDebugInfo)
+			{
+				DumpDebugShaderText(Input, ANSI_TO_TCHAR(SourceData.c_str()), TEXT("rewritten.hlsl"));
+			}
+		}
+		else
+		{
+			CompilerContext.FlushErrors(Output.Errors);
+			bCompilationFailed = true;
+		}
+		Options.bRemoveUnusedGlobals = false;
+	}
+
+	// Compile HLSL source to SPIR-V binary
+	TArray<uint32> SpirvData;
+	
+	if (!bCompilationFailed && !CompilerContext.CompileHlslToSpirv(Options, SpirvData))
+	{
+		// Flush compile errors
+		CompilerContext.FlushErrors(Output.Errors);
+		bCompilationFailed = true;
+	}
+
+	if (!bCompilationFailed)
+	{
+		ReflectionData ReflectData;
+		CrossCompiler::FHlslccHeaderWriter CCHeaderWriter;
+		
+		// Now perform reflection on the SPIRV and tweak any decorations that we need to.
+		// This used to be done via JSON, but that was slow and alloc happy so use SPIRV-Reflect instead.
+		spv_reflect::ShaderModule Reflection(SpirvData.Num() * sizeof(uint32), SpirvData.GetData());
+		check(Reflection.GetResult() == SPV_REFLECT_RESULT_SUCCESS);
+
+		const ANSICHAR* SPIRV_DummySamplerName = CrossCompiler::FShaderConductorContext::GetIdentifierTable().DummySampler;
+
+		ParseReflectionData(CCHeaderWriter, ReflectData, SpirvData, Reflection, SPIRV_DummySamplerName, Frequency, bEmulatedUBs);
+		
+		const ANSICHAR* FrequencyPrefix = GetFrequencyPrefix(Frequency);
+		
+		// Step 2 : End of reflection
+		// 
+		// Overwrite updated SPIRV code
+		SpirvData = TArray<uint32>(Reflection.GetCode(), Reflection.GetCodeSize() / 4);
+
+		if (bDumpDebugInfo)
+		{
+			// SPIR-V file (Binary)
+			DumpDebugShaderBinary(Input, SpirvData.GetData(), SpirvData.Num() * sizeof(uint32), TEXT("spv"));
+		}
+
+		CrossCompiler::FShaderConductorTarget TargetDesc;
+
+		switch (Version)
+		{
+		case GLSL_150_ES3_1:	// ES3.1 Emulation
+			TargetDesc.CompileFlags.SetDefine(TEXT("force_flattened_io_blocks"), 1);
+			TargetDesc.CompileFlags.SetDefine(TEXT("emit_uniform_buffer_as_plain_uniforms"), 1);
+			TargetDesc.Language = CrossCompiler::EShaderConductorLanguage::Glsl;
+			TargetDesc.Version = 430;
+			break;
+		case GLSL_SWITCH_FORWARD:
+			TargetDesc.Language = CrossCompiler::EShaderConductorLanguage::Essl;
+			TargetDesc.Version = 320;
+			break;
+		case GLSL_SWITCH:
+			TargetDesc.Language = CrossCompiler::EShaderConductorLanguage::Glsl;
+			TargetDesc.Version = 430;
+			break;
+		case GLSL_ES3_1_ANDROID:
+		default:
+			TargetDesc.CompileFlags.SetDefine(TEXT("force_flattened_io_blocks"), 1);
+			TargetDesc.CompileFlags.SetDefine(TEXT("emit_uniform_buffer_as_plain_uniforms"), 1);
+			TargetDesc.CompileFlags.SetDefine(TEXT("force_temporary"), 1);
+
+			// If we have mobile multiview define set then set the view count and enable extension
+			const FString* MultiViewDefine = Input.Environment.GetDefinitions().Find(TEXT("MOBILE_MULTI_VIEW"));
+			if (Frequency == SF_Vertex && MultiViewDefine && *MultiViewDefine == "1")
+			{
+				TargetDesc.CompileFlags.SetDefine(TEXT("ovr_multiview_view_count"), 2);
+			}
+
+			TargetDesc.Language = CrossCompiler::EShaderConductorLanguage::Essl;
+			TargetDesc.Version = 320;
+			break;
+		}
+
+		TSet<FString> ExternalTextures;
+		int32 Pos = 0;
+#if !PLATFORM_MAC
+		TCHAR TextureExternalName[256];
+#else
+		ANSICHAR TextureExternalName[256];
+#endif
+		do
+		{
+			Pos = PreprocessedShader.Find(TEXT("TextureExternal"), ESearchCase::CaseSensitive, ESearchDir::FromStart, Pos + 15);
+			if (Pos != INDEX_NONE)
+			{
+#if PLATFORM_WINDOWS
+				if (swscanf_s(&PreprocessedShader[Pos], TEXT("TextureExternal %ls"), TextureExternalName, 256) == 1)
+#elif PLATFORM_MAC
+				if (sscanf(TCHAR_TO_ANSI(&PreprocessedShader[Pos]), "TextureExternal %s", TextureExternalName) == 1)
+#else // PLATFORM_LINUX
+				if (swscanf(TCHAR_TO_WCHAR(&PreprocessedShader[Pos]), L"TextureExternal %ls", TextureExternalName) == 1)
+#endif
+				{
+					FString Name = TextureExternalName;
+					if (Name.RemoveFromEnd(TEXT(";")))
 					{
-						GlslSource.erase(UniformTypePos, OldUniformTypeName.length());
-						GlslSource.insert(UniformTypePos, NewUniformTypeName);
+						ExternalTextures.Add(TEXT("SPIRV_Cross_Combined") + Name);
 					}
 				}
+			}
+		}
+		while (Pos != INDEX_NONE);
 
-				// Generate meta data for CCHeader
-				CCHeaderWriter.WriteSourceInfo(*Input.GetSourceFilename(), *Input.EntryPointName, *Input.DebugGroupName);
-				CCHeaderWriter.WriteCompilerInfo();
+		// Define type renaming callback after all external texture types have been gathered
+		TargetDesc.VariableTypeRenameCallback = [&ExternalTextures](const FAnsiStringView& VariableName, const FAnsiStringView& TypeName, FString& OutRenamedTypeName) -> bool
+		{
+			for (const FString& ExternalTex : ExternalTextures)
+			{
+				if (VariableName.Len() == ExternalTex.Len() && FCStringWide::Strncmp(ANSI_TO_TCHAR(VariableName.GetData()), *ExternalTex, ExternalTex.Len()) == 0)
+				{
+					OutRenamedTypeName = TEXT("samplerExternalOES");
+					return true;
+				}
+			}
+			return false;
+		};
 
-				const FString MetaData = CCHeaderWriter.ToString();
+		GLSLCompileParameters GLSLCompileParams;
 
-				// Merge meta data and GLSL source to output string
-				const int32 GlslShaderSourceLen = MetaData.Len() + static_cast<int32>(GlslSource.size()) + 1;
-				OutGlslShaderSource = (char*)malloc(GlslShaderSourceLen);
-				FCStringAnsi::Snprintf(OutGlslShaderSource, GlslShaderSourceLen, "%s%s", TCHAR_TO_ANSI(*MetaData), GlslSource.c_str());
-			}
+		GLSLCompileParams.CompilerContext = &CompilerContext;
+		GLSLCompileParams.CCHeaderWriter = &CCHeaderWriter;
+		GLSLCompileParams.TargetDesc = &TargetDesc;
+		GLSLCompileParams.Options = &Options;
+
+		GLSLCompileParams.Output = &Output;
+		GLSLCompileParams.SpirvData = &SpirvData;
+
+		GLSLCompileParams.Frequency = Frequency;
+		GLSLCompileParams.SPIRV_DummySamplerName = SPIRV_DummySamplerName;
+
+		std::string GlslSource;
+
+		// Handle PLS and FBF in OpenGL
+		if (Input.Environment.GetDefinitions().Contains("SHADING_PATH_MOBILE") && Input.Environment.GetDefinitions()["SHADING_PATH_MOBILE"] == "1" &&
+			Input.Environment.GetDefinitions().Contains("MOBILE_DEFERRED_SHADING") && Input.Environment.GetDefinitions()["MOBILE_DEFERRED_SHADING"] == "1" &&
+			Version == GLSL_ES3_1_ANDROID)
+		{
+			bCompilationFailed = !GenerateDeferredMobileShaders(GlslSource, GLSLCompileParams, SourceData, ReflectData, true, false, bEmulatedUBs);
+
+		}
+		else
+		{
+			bCompilationFailed = !GenerateGlslShader(GlslSource, GLSLCompileParams, ReflectData, true, false, bEmulatedUBs);
+		}
+
+		// Generate meta data for CCHeader
+		CCHeaderWriter.WriteSourceInfo(*Input.GetSourceFilename(), *Input.EntryPointName, *Input.DebugGroupName);
+		CCHeaderWriter.WriteCompilerInfo();
+
+		if (GlslSource.length() > 0)
+		{
+			const FString MetaData = CCHeaderWriter.ToString();
+
+			// Merge meta data and GLSL source to output string
+			const int32 GlslShaderSourceLen = MetaData.Len() + static_cast<int32>(GlslSource.size()) + 1;
+			OutGlslShaderSource = (char*)malloc(GlslShaderSourceLen);
+			FCStringAnsi::Snprintf(OutGlslShaderSource, GlslShaderSourceLen, "%s%s", TCHAR_TO_ANSI(*MetaData), GlslSource.c_str());
 		}
 	}
 
diff --git a/Engine/Source/Runtime/OpenGLDrv/Private/OpenGLRenderTarget.cpp b/Engine/Source/Runtime/OpenGLDrv/Private/OpenGLRenderTarget.cpp
--- a/Engine/Source/Runtime/OpenGLDrv/Private/OpenGLRenderTarget.cpp
+++ b/Engine/Source/Runtime/OpenGLDrv/Private/OpenGLRenderTarget.cpp
@@ -1160,7 +1160,6 @@ void FOpenGLDynamicRHI::RHIBeginRenderPass(const FRHIRenderPassInfo& InInfo, con
 
 #if PLATFORM_ANDROID
 	if (RenderPassInfo.SubpassHint == ESubpassHint::DeferredShadingSubpass &&
-		!FOpenGL::SupportsShaderMRTFramebufferFetch() &&
 		FOpenGL::SupportsPixelLocalStorage())
 	{
 		glEnable(GL_SHADER_PIXEL_LOCAL_STORAGE_EXT);
@@ -1219,7 +1218,6 @@ void FOpenGLDynamicRHI::RHIEndRenderPass()
 
 #if PLATFORM_ANDROID
 	if (RenderPassInfo.SubpassHint == ESubpassHint::DeferredShadingSubpass &&
-		!FOpenGL::SupportsShaderMRTFramebufferFetch() &&
 		FOpenGL::SupportsPixelLocalStorage())
 	{
 		glDisable(GL_SHADER_PIXEL_LOCAL_STORAGE_EXT);
diff --git a/Engine/Source/Runtime/OpenGLDrv/Private/OpenGLShaders.cpp b/Engine/Source/Runtime/OpenGLDrv/Private/OpenGLShaders.cpp
--- a/Engine/Source/Runtime/OpenGLDrv/Private/OpenGLShaders.cpp
+++ b/Engine/Source/Runtime/OpenGLDrv/Private/OpenGLShaders.cpp
@@ -878,13 +878,13 @@ void OPENGLDRV_API GLSLToDeviceCompatibleGLSL(FAnsiCharArray& GlslCodeOriginal,
 	// This is the place to insert such engine preprocessor defines, immediately after the glsl version declaration.
 	if (TypeEnum == GL_FRAGMENT_SHADER)
 	{
-		if(FOpenGL::SupportsShaderMRTFramebufferFetch())
+		if (FOpenGL::SupportsPixelLocalStorage())
 		{
-			AppendCString(GlslCode, "#define UE_MRT_FRAMEBUFFER_FETCH 1\n");
+			AppendCString(GlslCode, "#define UE_MRT_PLS 1\n");
 		}
-		else if (FOpenGL::SupportsPixelLocalStorage())
+		else if(FOpenGL::SupportsShaderMRTFramebufferFetch())
 		{
-			AppendCString(GlslCode, "#define UE_MRT_PLS 1\n");
+			AppendCString(GlslCode, "#define UE_MRT_FRAMEBUFFER_FETCH 1\n");
 		}
 	}
 
diff --git a/Engine/Source/Runtime/RenderCore/Private/RenderUtils.cpp b/Engine/Source/Runtime/RenderCore/Private/RenderUtils.cpp
--- a/Engine/Source/Runtime/RenderCore/Private/RenderUtils.cpp
+++ b/Engine/Source/Runtime/RenderCore/Private/RenderUtils.cpp
@@ -1123,6 +1123,10 @@ RENDERCORE_API bool MobileRequiresSceneDepthAux(const FStaticShaderPlatform Plat
 	{
 		return true;
 	}
+	else if (IsMobileDeferredShadingEnabled(Platform) && IsAndroidOpenGLESPlatform(Platform))
+	{
+		return true;
+	}
 	else if (!IsMobileDeferredShadingEnabled(Platform) && bMobileHDR)
 	{
 		// SceneDepthAux disabled when MobileHDR=false for non-IOS
diff --git a/Engine/Source/Runtime/Renderer/Private/MobileShadingRenderer.cpp b/Engine/Source/Runtime/Renderer/Private/MobileShadingRenderer.cpp
--- a/Engine/Source/Runtime/Renderer/Private/MobileShadingRenderer.cpp
+++ b/Engine/Source/Runtime/Renderer/Private/MobileShadingRenderer.cpp
@@ -1278,7 +1278,7 @@ void FMobileSceneRenderer::RenderDeferred(FRDGBuilder& GraphBuilder, const FSort
 	TArray<FRDGTextureRef, TInlineAllocator<5>> ColorTargets;
 
 	// If we are using GL and don't have FBF support, use PLS
-	bool bUsingPixelLocalStorage = IsAndroidOpenGLESPlatform(ShaderPlatform) && GSupportsPixelLocalStorage && !GSupportsShaderMRTFramebufferFetch;
+	bool bUsingPixelLocalStorage = IsAndroidOpenGLESPlatform(ShaderPlatform) && GSupportsPixelLocalStorage;
 
 	if (bUsingPixelLocalStorage)
 	{
@@ -1371,22 +1371,31 @@ void FMobileSceneRenderer::RenderDeferredSinglePass(FRDGBuilder& GraphBuilder, c
 		RHICmdList.SetCurrentStat(GET_STATID(STAT_CLMM_Opaque));
 		RenderMobileBasePass(RHICmdList, View);
 		RHICmdList.ImmediateFlush(EImmediateFlushType::DispatchToRHIThread);
+		
 		// Issue occlusion queries
-		RHICmdList.SetCurrentStat(GET_STATID(STAT_CLMM_Occlusion));
-		RenderOcclusion(RHICmdList, View);
+		if (!bUsingPixelLocalStorage)
+		{
+			RHICmdList.SetCurrentStat(GET_STATID(STAT_CLMM_Occlusion));
+			RenderOcclusion(RHICmdList, View);
+		}
+
 		PostRenderBasePass(RHICmdList, View);
 		// SceneColor + GBuffer write, SceneDepth is read only
 		RHICmdList.NextSubpass();
 		RenderDecals(RHICmdList, View);
 		// SceneColor write, SceneDepth is read only
 		RHICmdList.NextSubpass();
+		MobileDeferredShadingPass(RHICmdList, ViewIndex, NumViews, View, *Scene, SortedLightSet);
+
 		if (bUsingPixelLocalStorage)
 		{
 			MobileDeferredCopyBuffer<FMobileDeferredCopyPLSPS>(RHICmdList, View);
 			// SceneColor write, SceneDepth is read only
 			RHICmdList.NextSubpass();
+
+			RHICmdList.SetCurrentStat(GET_STATID(STAT_CLMM_Occlusion));
+			RenderOcclusion(RHICmdList, View);
 		}
-		MobileDeferredShadingPass(RHICmdList, ViewIndex, NumViews, View, *Scene, SortedLightSet);
 		RenderFog(RHICmdList, View);
 		// Draw translucency.
 		RenderTranslucency(RHICmdList, View);
