Disable ubo padding as it causes a compile error when generating offsets for certain structures

diff --git a/Engine/Source/Developer/ShaderFormatOpenGL/Private/OpenGLShaderCompiler.cpp b/Engine/Source/Developer/ShaderFormatOpenGL/Private/OpenGLShaderCompiler.cpp
--- a/Engine/Source/Developer/ShaderFormatOpenGL/Private/OpenGLShaderCompiler.cpp
+++ b/Engine/Source/Developer/ShaderFormatOpenGL/Private/OpenGLShaderCompiler.cpp
@@ -2945,7 +2945,6 @@ static bool CompileToGlslWithShaderConductor(
 			TargetDesc.CompileFlags.SetDefine(TEXT("force_flattened_io_blocks"), 1);
 			TargetDesc.CompileFlags.SetDefine(TEXT("emit_uniform_buffer_as_plain_uniforms"), 1);
 			TargetDesc.CompileFlags.SetDefine(TEXT("force_temporary"), 1);
-			TargetDesc.CompileFlags.SetDefine(TEXT("pad_ubo_blocks"), 1);
 
 			// If we have mobile multiview define set then set the view count and enable extension
 			const FString* MultiViewDefine = Input.Environment.GetDefinitions().Find(TEXT("MOBILE_MULTI_VIEW"));
