Fix for texelfetch parsing incorrect sampler
Fix for texture external

diff --git a/Engine/Shaders/Public/Platform/GL/GLCommon.ush b/Engine/Shaders/Public/Platform/GL/GLCommon.ush
--- a/Engine/Shaders/Public/Platform/GL/GLCommon.ush
+++ b/Engine/Shaders/Public/Platform/GL/GLCommon.ush
@@ -5,7 +5,7 @@
 =============================================================================*/
 
 // Update this GUID to improve shader recompilation for GL only shaders
-// GUID = E121241A-1AD3-416B-8F74-288560D5F0005
+// GUID = FE909406-6AB7-4474-A490-709DB04F3312
 
 #pragma once
 
diff --git a/Engine/Source/Developer/ShaderFormatOpenGL/Private/OpenGLShaderCompiler.cpp b/Engine/Source/Developer/ShaderFormatOpenGL/Private/OpenGLShaderCompiler.cpp
--- a/Engine/Source/Developer/ShaderFormatOpenGL/Private/OpenGLShaderCompiler.cpp
+++ b/Engine/Source/Developer/ShaderFormatOpenGL/Private/OpenGLShaderCompiler.cpp
@@ -2462,17 +2462,9 @@ bool GenerateGlslShader(std::string& OutString, GLSLCompileParameters& GLSLCompi
 				OutString.erase(SamplerBufferPos + SamplerBufferName.size() - Texture.Len(), Texture.Len());
 				OutString.insert(SamplerBufferPos + SamplerBufferName.size() - Texture.Len(), NewSamplerBufferName);
 
-				std::string SamplerTexFetchName = std::string("texelFetch(") + TCHAR_TO_ANSI(*Texture);
-				size_t SamplerTexFetchPos = OutString.find(SamplerTexFetchName);
-
-				while (SamplerTexFetchPos != std::string::npos)
-				{
-
-					OutString.erase(SamplerTexFetchPos + SamplerTexFetchName.size() - Texture.Len(), Texture.Len());
-					OutString.insert(SamplerTexFetchPos + SamplerTexFetchName.size() - Texture.Len(), NewSamplerBufferName);
-
-					SamplerTexFetchPos = OutString.find(SamplerTexFetchName);
-				}
+				std::string SamplerTexFetchExpression = std::string("texelFetch\\(\\s?") + TCHAR_TO_ANSI(*Texture) + "\\s?,";
+				std::regex SamplerTexFetchPattern(SamplerTexFetchExpression);
+				OutString = std::regex_replace(OutString, SamplerTexFetchPattern, std::string("texelFetch(") + NewSamplerBufferName + ",");
 			}
 
 			const uint32 SamplerCount = FMath::Max(1, UsedSamplers.Num());
@@ -2985,7 +2977,7 @@ static bool CompileToGlslWithShaderConductor(
 					FString Name = TextureExternalName;
 					if (Name.RemoveFromEnd(TEXT(";")))
 					{
-						ExternalTextures.Add(TEXT("SPIRV_Cross_Combined") + Name);
+						ExternalTextures.Add(TEXT("SPIRV_Cross_Combined") + Name + Name + TEXT("Sampler"));
 					}
 				}
 			}
