Fix for OpenGL ES depth fetch when used without deferred rendering.

diff --git a/Engine/Shaders/Public/Platform/GL/GLCommon.ush b/Engine/Shaders/Public/Platform/GL/GLCommon.ush
--- a/Engine/Shaders/Public/Platform/GL/GLCommon.ush
+++ b/Engine/Shaders/Public/Platform/GL/GLCommon.ush
@@ -5,7 +5,7 @@
 =============================================================================*/
 
 // Update this GUID to improve shader recompilation for GL only shaders
-// GUID = FE909406-6AB7-4474-A490-709DB04F3312
+// GUID = 95D26A8C-47D1-431F-A949-5546C6A59D92
 
 #pragma once
 
diff --git a/Engine/Source/Developer/ShaderFormatOpenGL/Private/OpenGLShaderCompiler.cpp b/Engine/Source/Developer/ShaderFormatOpenGL/Private/OpenGLShaderCompiler.cpp
--- a/Engine/Source/Developer/ShaderFormatOpenGL/Private/OpenGLShaderCompiler.cpp
+++ b/Engine/Source/Developer/ShaderFormatOpenGL/Private/OpenGLShaderCompiler.cpp
@@ -2260,34 +2260,66 @@ bool GenerateGlslShader(std::string& OutString, GLSLCompileParameters& GLSLCompi
 					std::string DepthBufferIndex = bIsDeferred ? "4" : "1";
 					std::string DepthBufferOutVarString = "out_var_SV_Target" + DepthBufferIndex;
 
+					// Insert function declaration to handle retrieving depth
 					OutString.insert(MainPos, "float GLFetchDepthBuffer()\n"
 						"{\n"
-						"\t#ifdef GL_ARM_shader_framebuffer_fetch_depth_stencil\n"
+						"\t#if defined(GL_ARM_shader_framebuffer_fetch_depth_stencil)\n"
 						"\treturn gl_LastFragDepthARM;\n"
-						"\t#else\n"
+						"\t#elif defined(GL_EXT_shader_framebuffer_fetch)\n"
 						"\treturn " + DepthBufferOutVarString + ".x;\n"
+						"\t#else\n"
+						"\treturn 0.0f;\n"
 						"\t#endif\n"
 						"}\n");
 
+					
+					// If SceneDepthAux is not declared then declare it, otherwise modify so that we only enable it on devices that don't support
+					// GL_ARM_shader_framebuffer_fetch_depth_stencil and do support GL_EXT_shader_framebuffer_fetch
 					size_t DepthBufferOutVarPos = OutString.find(DepthBufferOutVarString + ";");
+
+					std::string DepthBufferDeclString = "layout(location = " + DepthBufferIndex + ") inout highp vec4 " + DepthBufferOutVarString + ";\n";
+					std::string DepthBufferOutString = "\n#if !defined(GL_ARM_shader_framebuffer_fetch_depth_stencil) && defined(GL_EXT_shader_framebuffer_fetch)\n" +
+														DepthBufferDeclString +
+														"#endif\n";
+
+					// If we cannot find a declararation of out_var_SV_Target(n) in the shader, insert one
 					if (DepthBufferOutVarPos == std::string::npos)
 					{
-						OutString.insert(MainPos, "\n#ifndef GL_ARM_shader_framebuffer_fetch_depth_stencil\n"
-							"layout(location = " + DepthBufferIndex + ") inout highp vec4 " + DepthBufferOutVarString + "; \n"
-							"#endif\n");
+						OutString.insert(MainPos, DepthBufferOutString);
 					}
-					else
+					else 
 					{
-						size_t QualifierPos = OutString.find_last_of("out", DepthBufferOutVarPos - 1);
-						OutString.erase(QualifierPos - 2, 3);
-						OutString.insert(QualifierPos - 2, "inout");
+						// If we have a declaration, replace with one that will be stripped if GL_ARM_shader_framebuffer_fetch_depth_stencil is enabled
+						size_t StringStartPos = OutString.rfind("layout", DepthBufferOutVarPos - 1);
+						size_t StringEndPos = OutString.find(";", StringStartPos);
+
+						OutString.erase(StringStartPos, (StringEndPos+1) - StringStartPos);
+						OutString.insert(StringStartPos, DepthBufferOutString);
+					}
+
+					// Make SceneDepthAux assignment conditional
+					// We only need to write the depth when we don't support GL_ARM_shader_framebuffer_fetch_depth_stencil
+					std::string DepthBufferAssignment = DepthBufferOutVarString + " =";
+					size_t DepthBufferAssignmentPos = OutString.find(DepthBufferAssignment);
+
+					if (DepthBufferAssignmentPos != std::string::npos)
+					{
+						size_t LineEnd = OutString.find_first_of(";", DepthBufferAssignmentPos);
+						uint32_t AssignmentValueStart = DepthBufferAssignmentPos + DepthBufferAssignment.size();
+						std::string AssignmentValue = OutString.substr(AssignmentValueStart + 1, LineEnd - AssignmentValueStart);
+
+						if (LineEnd != std::string::npos)
+						{
+							OutString.erase(DepthBufferAssignmentPos, LineEnd + 1 - DepthBufferAssignmentPos);
+							OutString.insert(DepthBufferAssignmentPos, std::string("#if !defined(GL_ARM_shader_framebuffer_fetch_depth_stencil) && defined(GL_EXT_shader_framebuffer_fetch)\n") + DepthBufferAssignment + AssignmentValue + std::string("\n#endif\n"));
+						}
 					}
 				}
 			}
 		}
 	}
 
-	// If we are rendering deferred, then we only need SceneDepthAux on devices that don't support framebuffer fetch depthj
+	// If we are rendering deferred, then we only need SceneDepthAux on devices that don't support framebuffer fetch depth
 	if (bIsDeferred)
 	{
 		std::string SceneDepthAux = "layout(location = 4) out highp float out_var_SV_Target4;";
