Fix static analysis warning

diff --git a/Engine/Source/Developer/ShaderFormatOpenGL/Private/GlslBackend.cpp b/Engine/Source/Developer/ShaderFormatOpenGL/Private/GlslBackend.cpp
--- a/Engine/Source/Developer/ShaderFormatOpenGL/Private/GlslBackend.cpp
+++ b/Engine/Source/Developer/ShaderFormatOpenGL/Private/GlslBackend.cpp
@@ -4266,7 +4266,7 @@ static ir_rvalue* GenShaderInputSemantic(
 	if (Variable == NULL && (Frequency == HSF_VertexShader || Frequency == HSF_PixelShader))
 	{
 		const int PrefixLength = 9;
-		if (FCStringAnsi::Strnicmp(Semantic, "SV_ViewID", PrefixLength) == 0)
+		if (Semantic && FCStringAnsi::Strnicmp(Semantic, "SV_ViewID", PrefixLength) == 0)
 		{
 			Variable = new(ParseState)ir_variable(
 				Type,
