Disabled force temporaries in OpenGLES due to bug when copying uniform buffer arrays with emulated uniform buffers

diff --git a/Engine/Shaders/Public/Platform/GL/GLCommon.ush b/Engine/Shaders/Public/Platform/GL/GLCommon.ush
--- a/Engine/Shaders/Public/Platform/GL/GLCommon.ush
+++ b/Engine/Shaders/Public/Platform/GL/GLCommon.ush
@@ -5,7 +5,7 @@
 =============================================================================*/
 
 // Update this GUID to improve shader recompilation for GL only shaders
-// GUID = 95D26A8C-47D1-431F-A949-5546C6A59D92
+// GUID = 73679724-6A97-4DE7-9727-4BEEF8C30BA3
 
 #pragma once
 
diff --git a/Engine/Source/Developer/ShaderFormatOpenGL/Private/OpenGLShaderCompiler.cpp b/Engine/Source/Developer/ShaderFormatOpenGL/Private/OpenGLShaderCompiler.cpp
--- a/Engine/Source/Developer/ShaderFormatOpenGL/Private/OpenGLShaderCompiler.cpp
+++ b/Engine/Source/Developer/ShaderFormatOpenGL/Private/OpenGLShaderCompiler.cpp
@@ -2974,7 +2974,8 @@ static bool CompileToGlslWithShaderConductor(
 		default:
 			TargetDesc.CompileFlags.SetDefine(TEXT("force_flattened_io_blocks"), 1);
 			TargetDesc.CompileFlags.SetDefine(TEXT("emit_uniform_buffer_as_plain_uniforms"), 1);
-			TargetDesc.CompileFlags.SetDefine(TEXT("force_temporary"), 1);
+			// TODO: Currently disabled due to bug when assigning an array to temporary variable
+			///TargetDesc.CompileFlags.SetDefine(TEXT("force_temporary"), 1);
 
 			// If we have mobile multiview define set then set the view count and enable extension
 			const FString* MultiViewDefine = Input.Environment.GetDefinitions().Find(TEXT("MOBILE_MULTI_VIEW"));
