Added support for decal blending with PLS on OpenGL

diff --git a/Engine/Shaders/Public/Platform/GL/GLCommon.ush b/Engine/Shaders/Public/Platform/GL/GLCommon.ush
--- a/Engine/Shaders/Public/Platform/GL/GLCommon.ush
+++ b/Engine/Shaders/Public/Platform/GL/GLCommon.ush
@@ -5,7 +5,7 @@
 =============================================================================*/
 
 // Update this GUID to improve shader recompilation for GL only shaders
-// GUID = ACE3B7A8-80A0-4585-B553-7E99B51F07C2
+// GUID = 73679724-6A97-4DE7-9727-4BEEF8C30BB4
 
 #pragma once
 
diff --git a/Engine/Source/Developer/ShaderFormatOpenGL/Private/OpenGLShaderCompiler.cpp b/Engine/Source/Developer/ShaderFormatOpenGL/Private/OpenGLShaderCompiler.cpp
--- a/Engine/Source/Developer/ShaderFormatOpenGL/Private/OpenGLShaderCompiler.cpp
+++ b/Engine/Source/Developer/ShaderFormatOpenGL/Private/OpenGLShaderCompiler.cpp
@@ -2619,7 +2619,151 @@ bool GenerateGlslShader(std::string& OutString, GLSLCompileParameters& GLSLCompi
 	return true;
 }
 
-bool GenerateDeferredMobileShaders(std::string& GlslSource, GLSLCompileParameters& GLSLCompileParams, const std::string& HlslString, ReflectionData& ReflectData, bool bWriteToCCHeader, bool bIsDeferred, bool bEmulatedUBs)
+enum EDecalBlendFlags
+{
+	DecalOut_MRT0  	= 1,
+	DecalOut_MRT1  	= 1 << 1,
+	DecalOut_MRT2  	= 1 << 2,
+	DecalOut_MRT3  	= 1 << 3,
+	Translucent		= 1 << 4,
+	AlphaComposite  = 1 << 5,
+	Modulate		= 1 << 6,
+};
+
+uint32_t GetDecalBlendFlags(const std::string& SourceString)
+{
+	uint32_t Flags = 0;
+
+	if (SourceString.find("DECAL_OUT_MRT0 1") != std::string::npos)
+	{
+		Flags |= EDecalBlendFlags::DecalOut_MRT0;
+	}
+	if (SourceString.find("DECAL_OUT_MRT1 1") != std::string::npos)
+	{
+		Flags |= EDecalBlendFlags::DecalOut_MRT1;
+	}
+	if (SourceString.find("DECAL_OUT_MRT2 1") != std::string::npos)
+	{
+		Flags |= EDecalBlendFlags::DecalOut_MRT2;
+	}
+	if (SourceString.find("DECAL_OUT_MRT3 1") != std::string::npos)
+	{
+		Flags |= EDecalBlendFlags::DecalOut_MRT3;
+	}
+
+	if (!Flags)
+	{
+		return 0;
+	}
+
+	if (SourceString.find("MATERIALBLENDING_ALPHACOMPOSITE 1") != std::string::npos)
+	{
+		Flags |= EDecalBlendFlags::AlphaComposite;
+	}
+	else if (SourceString.find("MATERIALBLENDING_MODULATE 1") != std::string::npos)
+	{
+		Flags |= EDecalBlendFlags::Modulate;
+	}
+	else if (SourceString.find("MATERIALBLENDING_TRANSLUCENT 1") != std::string::npos)
+	{
+		Flags |= EDecalBlendFlags::Translucent;
+	}
+
+	return Flags;
+}
+
+enum class EDecalBlendFunction
+{
+	SrcAlpha_One,
+	SrcAlpha_InvSrcAlpha,
+	DstColor_InvSrcAlpha,
+	One_InvSrcAlpha,
+	None
+};
+
+//																	Emissive,							Normal,										Metallic\Specular\Roughness,				BaseColor
+static const EDecalBlendFunction TranslucentBlendFunctions[]	= { EDecalBlendFunction::SrcAlpha_One,	EDecalBlendFunction::SrcAlpha_InvSrcAlpha,	EDecalBlendFunction::SrcAlpha_InvSrcAlpha,  EDecalBlendFunction::SrcAlpha_InvSrcAlpha };
+static const EDecalBlendFunction AlphaCompositeBlendFunctions[] = { EDecalBlendFunction::SrcAlpha_One,	EDecalBlendFunction::None,					EDecalBlendFunction::One_InvSrcAlpha,		EDecalBlendFunction::One_InvSrcAlpha };
+static const EDecalBlendFunction ModulateBlendFunctions[]		= { EDecalBlendFunction::SrcAlpha_One,	EDecalBlendFunction::SrcAlpha_InvSrcAlpha,	EDecalBlendFunction::SrcAlpha_InvSrcAlpha,  EDecalBlendFunction::DstColor_InvSrcAlpha };
+
+void GetDecalBlendFunctionString(const EDecalBlendFunction BlendFunction, const std::string& AttachmentString, const std::string& TempOutputName, std::string& OutputString)
+{
+	switch (BlendFunction)
+	{
+		case EDecalBlendFunction::SrcAlpha_One:
+		{
+			OutputString = AttachmentString + ".rgb = " + TempOutputName + ".a * " + TempOutputName + ".rgb + " + AttachmentString + ".rgb";
+			break;
+		}
+		case EDecalBlendFunction::SrcAlpha_InvSrcAlpha:
+		{
+			OutputString = AttachmentString + ".rgb = " + TempOutputName + ".a * " + TempOutputName + ".rgb + (1.0 - " + TempOutputName + ".a) * " + AttachmentString + ".rgb";
+			break;
+		}
+		case EDecalBlendFunction::DstColor_InvSrcAlpha:
+		{
+			OutputString = AttachmentString + ".rgb = " + TempOutputName + ".rgb * " + AttachmentString + ".rgb + (1.0 - " + TempOutputName + ".a) * " + AttachmentString + ".rgb";
+			break;
+		}
+		case EDecalBlendFunction::One_InvSrcAlpha:
+		{
+			OutputString = AttachmentString + ".rgb = " + TempOutputName + ".rgb + (1.0 - " + TempOutputName + ".a) * " + AttachmentString + ".rgb";
+			break;
+		}
+		case EDecalBlendFunction::None:
+		{
+			OutputString = "";
+			break;
+		}
+	}
+}
+
+void ModifyDecalBlending(std::string& SourceString, const uint32_t BlendFlags)
+{
+	for (uint32_t i = 0; i < 4; ++i)
+	{
+		std::string OutIndex = std::to_string(i);
+		std::string AttachmentString = "GENERATED_SubpassFetchAttachment" + OutIndex;
+		std::string AttachmentAssignmentString = AttachmentString + " =";
+
+		size_t AssignmentPos = SourceString.find(AttachmentAssignmentString);
+		if (AssignmentPos != std::string::npos)
+		{
+			size_t AssignmentEndPos = AssignmentPos + AttachmentAssignmentString.size();
+			size_t StringEndPos = SourceString.find(";", AssignmentEndPos);
+
+			std::string TempOutputName = "TempMulOut" + OutIndex;
+			std::string TempOut = "highp vec4 " + TempOutputName + " = " + SourceString.substr(AssignmentEndPos, StringEndPos - AssignmentEndPos) + "; \n";
+
+			SourceString.erase(AssignmentPos, StringEndPos - AssignmentPos);
+
+			if (1 << i & BlendFlags)
+			{
+				std::string BlendFnString;
+
+				if (BlendFlags & EDecalBlendFlags::Translucent)
+				{
+					GetDecalBlendFunctionString(TranslucentBlendFunctions[i], AttachmentString, TempOutputName, BlendFnString);
+				}
+				if (BlendFlags & EDecalBlendFlags::AlphaComposite)
+				{
+					GetDecalBlendFunctionString(AlphaCompositeBlendFunctions[i], AttachmentString, TempOutputName, BlendFnString);
+				}
+				if (BlendFlags & EDecalBlendFlags::Modulate)
+				{
+					GetDecalBlendFunctionString(ModulateBlendFunctions[i], AttachmentString, TempOutputName, BlendFnString);
+				}
+
+				TempOut += BlendFnString;
+
+				SourceString.insert(AssignmentPos, TempOut);
+			}			
+		}
+	}
+}
+
+bool GenerateDeferredMobileShaders(std::string& GlslSource, GLSLCompileParameters& GLSLCompileParams, const std::string& HlslString, ReflectionData& ReflectData,
+																		bool bWriteToCCHeader, bool bIsDeferred, bool bEmulatedUBs, uint32_t DecalBlendFlags)
 {
 	bool bHasGbufferTextures[4];
 	bHasGbufferTextures[0] = HlslString.find("GENERATED_SubpassFetchAttachment0") != std::string::npos;
@@ -2742,8 +2886,7 @@ bool GenerateDeferredMobileShaders(std::string& GlslSource, GLSLCompileParameter
 		}
 
 		// Replace assignment of output to buffer 0 to additive if the output proxy additive is used
-		if (HlslString.find("OutProxyAdditive") != std::string::npos ||
-			HlslString.find("OutTarget0") != std::string::npos)
+		if (HlslString.find("OutProxyAdditive") != std::string::npos)
 		{
 			std::string OutProxyString = "GENERATED_SubpassFetchAttachment0 =";
 			std::string OutProxyModifiedString = "GENERATED_SubpassFetchAttachment0 +=";
@@ -2756,6 +2899,11 @@ bool GenerateDeferredMobileShaders(std::string& GlslSource, GLSLCompileParameter
 			}
 		}
 
+		if (DecalBlendFlags)
+		{
+			ModifyDecalBlending(PLSSourceString, DecalBlendFlags);
+		}
+
 		// Strip version string
 		VersionStringPos = PLSSourceString.find(ESVersionString);
 		if (VersionStringPos != std::string::npos)
@@ -2863,6 +3011,8 @@ static bool CompileToGlslWithShaderConductor(
 		}
 	}
 
+	uint32_t BlendFlags = GetDecalBlendFlags(SourceData);
+
 	// Load shader source into compiler context
 	CompilerContext.LoadSource(SourceData.c_str(), FileName.c_str(), EntryPointName.c_str(), Frequency, &AdditionalDefines);
 
@@ -3028,7 +3178,7 @@ static bool CompileToGlslWithShaderConductor(
 			Input.Environment.GetDefinitions().Contains("MOBILE_DEFERRED_SHADING") && Input.Environment.GetDefinitions()["MOBILE_DEFERRED_SHADING"] == "1" &&
 			Version == GLSL_ES3_1_ANDROID)
 		{
-			bCompilationFailed = !GenerateDeferredMobileShaders(GlslSource, GLSLCompileParams, SourceData, ReflectData, true, false, bEmulatedUBs);
+			bCompilationFailed = !GenerateDeferredMobileShaders(GlslSource, GLSLCompileParams, SourceData, ReflectData, true, false, bEmulatedUBs, BlendFlags);
 		}
 		else
 		{
