Fixed bug with column/row mismatch on non-square matrices in OpenGL

diff --git a/Engine/Shaders/Public/Platform/GL/GLCommon.ush b/Engine/Shaders/Public/Platform/GL/GLCommon.ush
--- a/Engine/Shaders/Public/Platform/GL/GLCommon.ush
+++ b/Engine/Shaders/Public/Platform/GL/GLCommon.ush
@@ -5,7 +5,7 @@
 =============================================================================*/
 
 // Update this GUID to improve shader recompilation for GL only shaders
-// GUID = 73679724-6A97-4DE7-9727-4BEEF8C30BB4
+// GUID = E56FB0D5-6B8A-4DE3-800F-FFF7B9021A35
 
 #pragma once
 
diff --git a/Engine/Source/Developer/ShaderFormatOpenGL/Private/OpenGLShaderCompiler.cpp b/Engine/Source/Developer/ShaderFormatOpenGL/Private/OpenGLShaderCompiler.cpp
--- a/Engine/Source/Developer/ShaderFormatOpenGL/Private/OpenGLShaderCompiler.cpp
+++ b/Engine/Source/Developer/ShaderFormatOpenGL/Private/OpenGLShaderCompiler.cpp
@@ -1351,7 +1351,7 @@ void AddMemberToPackedUB(const std::string& FrequencyPrefix,
 			Name += Buff;
 		}
 
-		for (uint32_t i = 0; i < type.traits.numeric.matrix.row_count; ++i)
+		for (uint32_t i = 0; i < type.traits.numeric.matrix.column_count; ++i)
 		{
 			if (i > 0)
 			{
@@ -1365,7 +1365,7 @@ void AddMemberToPackedUB(const std::string& FrequencyPrefix,
 			std::string Buff = "[" + OffsetString + " + " + std::to_string(i) + "]";
 			Name += Buff;
 
-			switch (type.traits.numeric.matrix.column_count)
+			switch (type.traits.numeric.matrix.row_count)
 			{
 			case 0:
 			case 1:
