Fixed bug with deferred decals on OpenGL

diff --git a/Engine/Shaders/Public/Platform/GL/GLCommon.ush b/Engine/Shaders/Public/Platform/GL/GLCommon.ush
--- a/Engine/Shaders/Public/Platform/GL/GLCommon.ush
+++ b/Engine/Shaders/Public/Platform/GL/GLCommon.ush
@@ -5,7 +5,7 @@
 =============================================================================*/
 
 // Update this GUID to improve shader recompilation for GL only shaders
-// GUID = fa3b54fb-4438-4ec2-88b2-0dec90de6652
+// GUID = d5d137cf-9a2f-48cc-ade3-6a435fe6252c
 
 #pragma once
 
diff --git a/Engine/Source/Developer/ShaderFormatOpenGL/Private/OpenGLShaderCompiler.cpp b/Engine/Source/Developer/ShaderFormatOpenGL/Private/OpenGLShaderCompiler.cpp
--- a/Engine/Source/Developer/ShaderFormatOpenGL/Private/OpenGLShaderCompiler.cpp
+++ b/Engine/Source/Developer/ShaderFormatOpenGL/Private/OpenGLShaderCompiler.cpp
@@ -2633,23 +2633,35 @@ enum EDecalBlendFlags
 	Modulate		= 1 << 6,
 };
 
-uint32_t GetDecalBlendFlags(const std::string& SourceString)
+bool EnvironmentHasMatchingKeyValue(const TMap<FString, FString>& Definitions, const FString& Key, const FString& Value)
+{
+	const FString* MatchedKey = Definitions.Find(Key);
+	if (MatchedKey)
+	{
+		return *MatchedKey == Value;
+	}
+	return false;
+}
+
+uint32_t GetDecalBlendFlags(const FShaderCompilerInput& Input)
 {
 	uint32_t Flags = 0;
 
-	if (SourceString.find("DECAL_OUT_MRT0 1") != std::string::npos)
+	const TMap<FString, FString>& Definitions = Input.Environment.GetDefinitions();
+	
+	if (EnvironmentHasMatchingKeyValue(Definitions, TEXT("DECAL_OUT_MRT0"), TEXT("1")))
 	{
 		Flags |= EDecalBlendFlags::DecalOut_MRT0;
 	}
-	if (SourceString.find("DECAL_OUT_MRT1 1") != std::string::npos)
+	if (EnvironmentHasMatchingKeyValue(Definitions, TEXT("DECAL_OUT_MRT1"), TEXT("1")))
 	{
 		Flags |= EDecalBlendFlags::DecalOut_MRT1;
 	}
-	if (SourceString.find("DECAL_OUT_MRT2 1") != std::string::npos)
+	if (EnvironmentHasMatchingKeyValue(Definitions, TEXT("DECAL_OUT_MRT2"), TEXT("1")))
 	{
 		Flags |= EDecalBlendFlags::DecalOut_MRT2;
 	}
-	if (SourceString.find("DECAL_OUT_MRT3 1") != std::string::npos)
+	if (EnvironmentHasMatchingKeyValue(Definitions, TEXT("DECAL_OUT_MRT3"), TEXT("1")))
 	{
 		Flags |= EDecalBlendFlags::DecalOut_MRT3;
 	}
@@ -2659,15 +2671,15 @@ uint32_t GetDecalBlendFlags(const std::string& SourceString)
 		return 0;
 	}
 
-	if (SourceString.find("MATERIALBLENDING_ALPHACOMPOSITE 1") != std::string::npos)
+	if (EnvironmentHasMatchingKeyValue(Definitions, TEXT("MATERIALBLENDING_ALPHACOMPOSITE"), TEXT("1")))
 	{
 		Flags |= EDecalBlendFlags::AlphaComposite;
 	}
-	else if (SourceString.find("MATERIALBLENDING_MODULATE 1") != std::string::npos)
+	else if (EnvironmentHasMatchingKeyValue(Definitions, TEXT("MATERIALBLENDING_MODULATE"), TEXT("1")))
 	{
 		Flags |= EDecalBlendFlags::Modulate;
 	}
-	else if (SourceString.find("MATERIALBLENDING_TRANSLUCENT 1") != std::string::npos)
+	else if (EnvironmentHasMatchingKeyValue(Definitions, TEXT("MATERIALBLENDING_TRANSLUCENT"), TEXT("1")))
 	{
 		Flags |= EDecalBlendFlags::Translucent;
 	}
@@ -3002,7 +3014,7 @@ static bool CompileToGlslWithShaderConductor(
 	DebugDataOptions.HlslCCFlags = CCFlags;
 	UE::ShaderCompilerCommon::DumpDebugShaderData(Input, PreprocessedShader, DebugDataOptions);
 
-	uint32_t BlendFlags = GetDecalBlendFlags(SourceData);
+	uint32_t BlendFlags = GetDecalBlendFlags(Input);
 
 	// Load shader source into compiler context
 	CompilerContext.LoadSource(SourceData.c_str(), FileName.c_str(), EntryPointName.c_str(), Frequency, &AdditionalDefines);
@@ -3176,8 +3188,9 @@ static bool CompileToGlslWithShaderConductor(
 		std::string GlslSource;
 
 		// Handle PLS and FBF in OpenGL
-		if (Input.Environment.GetDefinitions().Contains("SHADING_PATH_MOBILE") && Input.Environment.GetDefinitions()["SHADING_PATH_MOBILE"] == "1" &&
-			Input.Environment.GetDefinitions().Contains("MOBILE_DEFERRED_SHADING") && Input.Environment.GetDefinitions()["MOBILE_DEFERRED_SHADING"] == "1" &&
+
+		if (EnvironmentHasMatchingKeyValue(Input.Environment.GetDefinitions(), TEXT("SHADING_PATH_MOBILE"), TEXT("1")) &&
+			EnvironmentHasMatchingKeyValue(Input.Environment.GetDefinitions(), TEXT("MOBILE_DEFERRED_SHADING"), TEXT("1")) &&
 			Version == GLSL_ES3_1_ANDROID)
 		{
 			bCompilationFailed = !GenerateDeferredMobileShaders(GlslSource, GLSLCompileParams, SourceData, ReflectData, true, false, bEmulatedUBs, BlendFlags);
