ShaderMinifier - Add support for dead code removal to all shader formats

* Add code to all shader format back-ends to take dead code removal into account when computing the version ID
* Not enabled in Metal and OpenGL back-ends (require further testing)

diff --git a/Engine/Source/Developer/Apple/MetalShaderFormat/Private/MetalShaderCompiler.cpp b/Engine/Source/Developer/Apple/MetalShaderFormat/Private/MetalShaderCompiler.cpp
--- a/Engine/Source/Developer/Apple/MetalShaderFormat/Private/MetalShaderCompiler.cpp
+++ b/Engine/Source/Developer/Apple/MetalShaderFormat/Private/MetalShaderCompiler.cpp
@@ -1078,6 +1078,14 @@ void CompileShader_Metal(const FShaderCompilerInput& _Input,FShaderCompilerOutpu
 	// Process TEXT macro.
 	TransformStringIntoCharacterArray(PreprocessedShader);
 
+	// Run the experimental shader minifier
+	#if UE_METAL_SHADER_COMPILER_ALLOW_DEAD_CODE_REMOVAL
+	if (Input.Environment.CompilerFlags.Contains(CFLAG_RemoveDeadCode))
+	{
+		UE::ShaderCompilerCommon::RemoveDeadCode(PreprocessedShader, Input.EntryPointName, Output.Errors);
+	}
+	#endif // UE_METAL_SHADER_COMPILER_ALLOW_DEAD_CODE_REMOVAL
+
 	Output.PreprocessTime = FPlatformTime::Seconds() - StartPreprocessTime;
 	
 	uint32 CCFlags = HLSLCC_NoPreprocess | HLSLCC_PackUniformsIntoUniformBufferWithNames | HLSLCC_FixAtomicReferences | HLSLCC_RetainSizes | HLSLCC_KeepSamplerAndImageNames;
diff --git a/Engine/Source/Developer/Apple/MetalShaderFormat/Private/MetalShaderFormat.cpp b/Engine/Source/Developer/Apple/MetalShaderFormat/Private/MetalShaderFormat.cpp
--- a/Engine/Source/Developer/Apple/MetalShaderFormat/Private/MetalShaderFormat.cpp
+++ b/Engine/Source/Developer/Apple/MetalShaderFormat/Private/MetalShaderFormat.cpp
@@ -338,7 +338,19 @@ uint32 GetMetalFormatVersion(FName Format)
 	check(Version.Version.Format == FMetalShaderFormat::HEADER_VERSION);
 	check(Version.Version.HLSLCCMinor == HLSLCC_VersionMinor);
 	
-	return Version.Raw;
+	uint32 Result = Version.Raw;
+
+#if UE_METAL_SHADER_COMPILER_ALLOW_DEAD_CODE_REMOVAL
+	{
+		static const auto CVar = IConsoleManager::Get().FindConsoleVariable(TEXT("r.Shader.RemoveDeadCode"));
+		if (CVar && CVar->GetInt() != 0)
+		{
+			Result = HashCombine(Result, 0x75E2FE85);
+		}
+	}
+#endif // UE_METAL_SHADER_COMPILER_ALLOW_DEAD_CODE_REMOVAL
+
+	return Result;
 }
 
 /**
diff --git a/Engine/Source/Developer/Apple/MetalShaderFormat/Private/MetalShaderFormat.h b/Engine/Source/Developer/Apple/MetalShaderFormat/Private/MetalShaderFormat.h
--- a/Engine/Source/Developer/Apple/MetalShaderFormat/Private/MetalShaderFormat.h
+++ b/Engine/Source/Developer/Apple/MetalShaderFormat/Private/MetalShaderFormat.h
@@ -11,6 +11,11 @@
 #include "RHIFeatureLevel.h"
 #include "RHIShaderPlatform.h"
 
+// Controls whether r.Shaders.RemoveDeadCode should be honored
+#ifndef UE_METAL_SHADER_COMPILER_ALLOW_DEAD_CODE_REMOVAL
+#define UE_METAL_SHADER_COMPILER_ALLOW_DEAD_CODE_REMOVAL 0 // Disabled initially, requires testing
+#endif // UE_METAL_SHADER_COMPILER_ALLOW_DEAD_CODE_REMOVAL
+
 // IOS and TVOS use the mobile toolchain.
 enum EAppleSDKType
 {
diff --git a/Engine/Source/Developer/ShaderFormatOpenGL/Private/OpenGLShaderCompiler.cpp b/Engine/Source/Developer/ShaderFormatOpenGL/Private/OpenGLShaderCompiler.cpp
--- a/Engine/Source/Developer/ShaderFormatOpenGL/Private/OpenGLShaderCompiler.cpp
+++ b/Engine/Source/Developer/ShaderFormatOpenGL/Private/OpenGLShaderCompiler.cpp
@@ -3351,6 +3351,14 @@ void FOpenGLFrontend::CompileShader(const FShaderCompilerInput& Input, FShaderCo
 	// Process TEXT macro.
 	TransformStringIntoCharacterArray(PreprocessedShader);
 
+	// Run the experimental shader minifier
+#if UE_OPENGL_SHADER_COMPILER_ALLOW_DEAD_CODE_REMOVAL
+	if (Input.Environment.CompilerFlags.Contains(CFLAG_RemoveDeadCode))
+	{
+		UE::ShaderCompilerCommon::RemoveDeadCode(PreprocessedShader, Input.EntryPointName, Output.Errors);
+	}
+#endif // UE_OPENGL_SHADER_COMPILER_ALLOW_DEAD_CODE_REMOVAL
+
 	uint32 CCFlags = CalculateCrossCompilerFlags(Version, Input.Environment.FullPrecisionInPS, Input.Environment.CompilerFlags);
 
 	// Required as we added the RemoveUniformBuffersFromSource() function (the cross-compiler won't be able to interpret comments w/o a preprocessor)
diff --git a/Engine/Source/Developer/ShaderFormatOpenGL/Private/ShaderFormatOpenGL.cpp b/Engine/Source/Developer/ShaderFormatOpenGL/Private/ShaderFormatOpenGL.cpp
--- a/Engine/Source/Developer/ShaderFormatOpenGL/Private/ShaderFormatOpenGL.cpp
+++ b/Engine/Source/Developer/ShaderFormatOpenGL/Private/ShaderFormatOpenGL.cpp
@@ -38,7 +38,19 @@ public:
 		{
 			check(0);
 		}
-		const uint16 Version = ((HLSLCC_VersionMinor & 0xff) << 8) | (GLSLVersion & 0xff);
+
+		uint32 Version = ((HLSLCC_VersionMinor & 0xff) << 8) | (GLSLVersion & 0xff);
+
+	#if UE_OPENGL_SHADER_COMPILER_ALLOW_DEAD_CODE_REMOVAL
+		{
+			static const auto CVar = IConsoleManager::Get().FindConsoleVariable(TEXT("r.Shader.RemoveDeadCode"));
+			if (CVar && CVar->GetInt() != 0)
+			{
+				Version = HashCombine(Version, 0x75E2FE85);
+			}
+		}
+	#endif // UE_OPENGL_SHADER_COMPILER_ALLOW_DEAD_CODE_REMOVAL
+
 		return Version;
 	}
 	virtual void GetSupportedFormats(TArray<FName>& OutFormats) const override
diff --git a/Engine/Source/Developer/ShaderFormatOpenGL/Public/ShaderFormatOpenGL.h b/Engine/Source/Developer/ShaderFormatOpenGL/Public/ShaderFormatOpenGL.h
--- a/Engine/Source/Developer/ShaderFormatOpenGL/Public/ShaderFormatOpenGL.h
+++ b/Engine/Source/Developer/ShaderFormatOpenGL/Public/ShaderFormatOpenGL.h
@@ -17,6 +17,11 @@ class FArchive;
 #include "Windows/WindowsHWrapper.h"
 #endif
 
+// Controls whether r.Shaders.RemoveDeadCode should be honored
+#ifndef UE_OPENGL_SHADER_COMPILER_ALLOW_DEAD_CODE_REMOVAL
+#define UE_OPENGL_SHADER_COMPILER_ALLOW_DEAD_CODE_REMOVAL 0 // Disabled initially, requires testing
+#endif // UE_OPENGL_SHADER_COMPILER_ALLOW_DEAD_CODE_REMOVAL
+
 class FShaderCompilerFlags;
 struct FShaderCompilerInput;
 struct FShaderCompilerOutput;
diff --git a/Engine/Source/Developer/VulkanShaderFormat/Private/VulkanShaderCompiler.cpp b/Engine/Source/Developer/VulkanShaderFormat/Private/VulkanShaderCompiler.cpp
--- a/Engine/Source/Developer/VulkanShaderFormat/Private/VulkanShaderCompiler.cpp
+++ b/Engine/Source/Developer/VulkanShaderFormat/Private/VulkanShaderCompiler.cpp
@@ -2511,6 +2511,14 @@ void DoCompileVulkanShader(const FShaderCompilerInput& Input, FShaderCompilerOut
 	// Process TEXT macro.
 	TransformStringIntoCharacterArray(PreprocessedShaderSource);
 
+	// Run the experimental shader minifier
+	#if UE_VULKAN_SHADER_COMPILER_ALLOW_DEAD_CODE_REMOVAL
+	if (Input.Environment.CompilerFlags.Contains(CFLAG_RemoveDeadCode))
+	{
+		UE::ShaderCompilerCommon::RemoveDeadCode(PreprocessedShaderSource, EntryPointName, Output.Errors);
+	}
+	#endif // UE_VULKAN_SHADER_COMPILER_ALLOW_DEAD_CODE_REMOVAL
+
 	Output.PreprocessTime = FPlatformTime::Seconds() - StartPreprocessTime;
 
 	FCompilerInfo CompilerInfo(Input, WorkingDirectory, HlslFrequency);
diff --git a/Engine/Source/Developer/VulkanShaderFormat/Private/VulkanShaderFormat.cpp b/Engine/Source/Developer/VulkanShaderFormat/Private/VulkanShaderFormat.cpp
--- a/Engine/Source/Developer/VulkanShaderFormat/Private/VulkanShaderFormat.cpp
+++ b/Engine/Source/Developer/VulkanShaderFormat/Private/VulkanShaderFormat.cpp
@@ -55,6 +55,16 @@ public:
 		Version = HashCombine(Version, 0xFC0848E2);
 	#endif
 
+	#if UE_VULKAN_SHADER_COMPILER_ALLOW_DEAD_CODE_REMOVAL
+		{
+			static const auto CVar = IConsoleManager::Get().FindConsoleVariable(TEXT("r.Shader.RemoveDeadCode"));
+			if (CVar && CVar->GetInt() != 0)
+			{
+				Version = HashCombine(Version, 0x75E2FE85);
+			}
+		}
+	#endif // UE_VULKAN_SHADER_COMPILER_ALLOW_DEAD_CODE_REMOVAL
+
 		return Version;
 	}
 	virtual void GetSupportedFormats(TArray<FName>& OutFormats) const
diff --git a/Engine/Source/Developer/VulkanShaderFormat/Private/VulkanShaderFormat.h b/Engine/Source/Developer/VulkanShaderFormat/Private/VulkanShaderFormat.h
--- a/Engine/Source/Developer/VulkanShaderFormat/Private/VulkanShaderFormat.h
+++ b/Engine/Source/Developer/VulkanShaderFormat/Private/VulkanShaderFormat.h
@@ -6,6 +6,12 @@
 #include "SpirvCommon.h"
 #include "hlslcc.h"
 
+
+// Controls whether r.Shaders.RemoveDeadCode should be honored
+#ifndef UE_VULKAN_SHADER_COMPILER_ALLOW_DEAD_CODE_REMOVAL
+#define UE_VULKAN_SHADER_COMPILER_ALLOW_DEAD_CODE_REMOVAL 1
+#endif // UE_VULKAN_SHADER_COMPILER_ALLOW_DEAD_CODE_REMOVAL
+
 enum class EVulkanShaderVersion
 {
 	ES3_1,
diff --git a/Engine/Source/Developer/Windows/ShaderFormatD3D/Private/D3DShaderCompiler.cpp b/Engine/Source/Developer/Windows/ShaderFormatD3D/Private/D3DShaderCompiler.cpp
--- a/Engine/Source/Developer/Windows/ShaderFormatD3D/Private/D3DShaderCompiler.cpp
+++ b/Engine/Source/Developer/Windows/ShaderFormatD3D/Private/D3DShaderCompiler.cpp
@@ -1248,11 +1248,12 @@ void CompileD3DShader(const FShaderCompilerInput& Input, FShaderCompilerOutput&
 	TArray<FString> FilteredErrors;
 
 	// Run the experimental shader minifier
-
+	#if UE_D3D_SHADER_COMPILER_ALLOW_DEAD_CODE_REMOVAL
 	if (Input.Environment.CompilerFlags.Contains(CFLAG_RemoveDeadCode))
 	{
 		UE::ShaderCompilerCommon::RemoveDeadCode(PreprocessedShaderSource, EntryPointName, Output.Errors);
 	}
+	#endif // UE_D3D_SHADER_COMPILER_ALLOW_DEAD_CODE_REMOVAL
 
 	// @TODO - implement different material path to allow us to remove backwards compat flag on sm5 shaders
 	uint32 CompileFlags = D3DCOMPILE_ENABLE_BACKWARDS_COMPATIBILITY
diff --git a/Engine/Source/Developer/Windows/ShaderFormatD3D/Private/ShaderFormatD3D.cpp b/Engine/Source/Developer/Windows/ShaderFormatD3D/Private/ShaderFormatD3D.cpp
--- a/Engine/Source/Developer/Windows/ShaderFormatD3D/Private/ShaderFormatD3D.cpp
+++ b/Engine/Source/Developer/Windows/ShaderFormatD3D/Private/ShaderFormatD3D.cpp
@@ -43,7 +43,18 @@ public:
 	inline uint32 GetVersionHash(EVersion InVersion) const
 	{
 		const uint32 BaseHash = GetTypeHash(UE_SHADER_PCD3D_SHARED_VER);
-		const uint32 VersionHash = GetTypeHash(InVersion);
+		uint32 VersionHash = GetTypeHash(InVersion);
+
+	#if UE_D3D_SHADER_COMPILER_ALLOW_DEAD_CODE_REMOVAL
+		{
+			static const auto CVar = IConsoleManager::Get().FindConsoleVariable(TEXT("r.Shader.RemoveDeadCode"));
+			if (CVar && CVar->GetInt() != 0)
+			{
+				VersionHash = HashCombine(VersionHash, 0x75E2FE85);
+			}
+		}
+	#endif // UE_D3D_SHADER_COMPILER_ALLOW_DEAD_CODE_REMOVAL
+
 		return HashCombine(BaseHash, VersionHash);
 	}
 
diff --git a/Engine/Source/Developer/Windows/ShaderFormatD3D/Private/ShaderFormatD3D.h b/Engine/Source/Developer/Windows/ShaderFormatD3D/Private/ShaderFormatD3D.h
--- a/Engine/Source/Developer/Windows/ShaderFormatD3D/Private/ShaderFormatD3D.h
+++ b/Engine/Source/Developer/Windows/ShaderFormatD3D/Private/ShaderFormatD3D.h
@@ -8,6 +8,11 @@
 // TODO: Lock D3D12 to SM 6.6 min spec
 #define USE_SHADER_MODEL_6_6 1
 
+// Controls whether r.Shaders.RemoveDeadCode should be honored
+#ifndef UE_D3D_SHADER_COMPILER_ALLOW_DEAD_CODE_REMOVAL
+#define UE_D3D_SHADER_COMPILER_ALLOW_DEAD_CODE_REMOVAL 1
+#endif // UE_D3D_SHADER_COMPILER_ALLOW_DEAD_CODE_REMOVAL
+
 struct FShaderTarget;
 
 enum class ELanguage
