AR Camera Passthrough fix with TextureExternal

diff --git a/Engine/Shaders/Public/Platform/GL/GLCommon.ush b/Engine/Shaders/Public/Platform/GL/GLCommon.ush
--- a/Engine/Shaders/Public/Platform/GL/GLCommon.ush
+++ b/Engine/Shaders/Public/Platform/GL/GLCommon.ush
@@ -5,7 +5,7 @@
 =============================================================================*/
 
 // Update this GUID to improve shader recompilation for GL only shaders
-// GUID = d5d137cf-9a2f-48cc-ade3-6a435fe6252c
+// GUID = 5c8c127e-42c7-4e59-b510-5b4570e94083
 
 #pragma once
 
diff --git a/Engine/Source/Developer/ShaderFormatOpenGL/Private/OpenGLShaderCompiler.cpp b/Engine/Source/Developer/ShaderFormatOpenGL/Private/OpenGLShaderCompiler.cpp
--- a/Engine/Source/Developer/ShaderFormatOpenGL/Private/OpenGLShaderCompiler.cpp
+++ b/Engine/Source/Developer/ShaderFormatOpenGL/Private/OpenGLShaderCompiler.cpp
@@ -3110,7 +3110,7 @@ static bool CompileToGlslWithShaderConductor(
 			TargetDesc.CompileFlags.SetDefine(TEXT("emit_uniform_buffer_as_plain_uniforms"), 1);
 			TargetDesc.CompileFlags.SetDefine(TEXT("pad_ubo_blocks"), 1);
 			// TODO: Currently disabled due to bug when assigning an array to temporary variable
-			///TargetDesc.CompileFlags.SetDefine(TEXT("force_temporary"), 1);
+			//TargetDesc.CompileFlags.SetDefine(TEXT("force_temporary"), 1);
 
 			// If we have mobile multiview define set then set the view count and enable extension
 			const FString* MultiViewDefine = Input.Environment.GetDefinitions().Find(TEXT("MOBILE_MULTI_VIEW"));
@@ -3151,7 +3151,7 @@ static bool CompileToGlslWithShaderConductor(
 					FString Name = TextureExternalName;
 					if (Name.RemoveFromEnd(TEXT(";")))
 					{
-						ExternalTextures.Add(TEXT("SPIRV_Cross_Combined") + Name + Name + TEXT("Sampler"));
+						ExternalTextures.Add(Name + TEXT("Sampler"));
 					}
 				}
 			}
@@ -3163,7 +3163,7 @@ static bool CompileToGlslWithShaderConductor(
 		{
 			for (const FString& ExternalTex : ExternalTextures)
 			{
-				if (VariableName.Len() == ExternalTex.Len() && FCStringWide::Strncmp(ANSI_TO_TCHAR(VariableName.GetData()), *ExternalTex, ExternalTex.Len()) == 0)
+				if (FCStringWide::Strstr(ANSI_TO_TCHAR(VariableName.GetData()), *ExternalTex))
 				{
 					OutRenamedTypeName = TEXT("samplerExternalOES");
 					return true;
