[Partial Backout] - CL24348577
Original CL Desc
-----------------------------------------------------------------
Remove unnecessary private include path

#rnx

diff --git a/Engine/Source/Developer/ShaderFormatOpenGL/ShaderFormatOpenGL.Build.cs b/Engine/Source/Developer/ShaderFormatOpenGL/ShaderFormatOpenGL.Build.cs
--- a/Engine/Source/Developer/ShaderFormatOpenGL/ShaderFormatOpenGL.Build.cs
+++ b/Engine/Source/Developer/ShaderFormatOpenGL/ShaderFormatOpenGL.Build.cs
@@ -7,12 +7,10 @@ public class ShaderFormatOpenGL : ModuleRules
 	public ShaderFormatOpenGL(ReadOnlyTargetRules Target) : base(Target)
 	{
 
-		PrivateIncludePathModuleNames.AddRange(
-			new string[] {
-				"TargetPlatform",
-				"OpenGLDrv"		
-			}
-			);
+		PrivateIncludePathModuleNames.Add("TargetPlatform");
+
+		PrivateIncludePaths.Add("Runtime/OpenGLDrv/Private");
+		PrivateIncludePaths.Add("Runtime/OpenGLDrv/Public");
 
 		PrivateDependencyModuleNames.AddRange(
 			new string[] {
diff --git a/Engine/Source/Developer/VulkanShaderFormat/VulkanShaderFormat.Build.cs b/Engine/Source/Developer/VulkanShaderFormat/VulkanShaderFormat.Build.cs
--- a/Engine/Source/Developer/VulkanShaderFormat/VulkanShaderFormat.Build.cs
+++ b/Engine/Source/Developer/VulkanShaderFormat/VulkanShaderFormat.Build.cs
@@ -10,7 +10,7 @@ public class VulkanShaderFormat : ModuleRules
 		PrivateIncludePathModuleNames.Add("TargetPlatform");
 
 		// Do not link the module (as that would require the vulkan dll), only the include paths
-		PublicIncludePathModuleNames.Add("VulkanRHI");
+		PublicIncludePaths.Add("Runtime/VulkanRHI/Public");
 
 		PrivateDependencyModuleNames.AddRange(
 			new string[] {
