Remove unnecessary private include path

diff --git a/Engine/Source/Developer/ShaderFormatOpenGL/ShaderFormatOpenGL.Build.cs b/Engine/Source/Developer/ShaderFormatOpenGL/ShaderFormatOpenGL.Build.cs
--- a/Engine/Source/Developer/ShaderFormatOpenGL/ShaderFormatOpenGL.Build.cs
+++ b/Engine/Source/Developer/ShaderFormatOpenGL/ShaderFormatOpenGL.Build.cs
@@ -6,11 +6,12 @@ public class ShaderFormatOpenGL : ModuleRules
 {
 	public ShaderFormatOpenGL(ReadOnlyTargetRules Target) : base(Target)
 	{
-
-		PrivateIncludePathModuleNames.Add("TargetPlatform");
-
-		PrivateIncludePaths.Add("Runtime/OpenGLDrv/Private");
-		PrivateIncludePaths.Add("Runtime/OpenGLDrv/Public");
+		PrivateIncludePathModuleNames.AddRange(
+			new string[] {
+				"TargetPlatform",
+				"OpenGLDrv"
+			}
+			);
 
 		PrivateDependencyModuleNames.AddRange(
 			new string[] {
