Fixed bug in previous CL with fma define inserted before verion on desktop GL
Fixed bug where sampler names were matched incorrectly

diff --git a/Engine/Shaders/Public/Platform/GL/GLCommon.ush b/Engine/Shaders/Public/Platform/GL/GLCommon.ush
--- a/Engine/Shaders/Public/Platform/GL/GLCommon.ush
+++ b/Engine/Shaders/Public/Platform/GL/GLCommon.ush
@@ -5,7 +5,7 @@
 =============================================================================*/
 
 // Update this GUID to improve shader recompilation for GL only shaders
-// GUID = C8B835E4-AF1B-4D8E-BD49-CBAAA96D6A0B
+// GUID = 7675143E-B3FA-41CE-ADD7-7C766E98E87F
 
 #pragma once
 
diff --git a/Engine/Source/Developer/ShaderFormatOpenGL/Private/OpenGLShaderCompiler.cpp b/Engine/Source/Developer/ShaderFormatOpenGL/Private/OpenGLShaderCompiler.cpp
--- a/Engine/Source/Developer/ShaderFormatOpenGL/Private/OpenGLShaderCompiler.cpp
+++ b/Engine/Source/Developer/ShaderFormatOpenGL/Private/OpenGLShaderCompiler.cpp
@@ -2467,7 +2467,8 @@ bool GenerateGlslShader(std::string& OutString, GLSLCompileParameters& GLSLCompi
 			{
 				std::string SamplerName = "SPIRV_Cross_Combined";
 				SamplerName += TCHAR_TO_ANSI(*(Texture + Sampler));
-				size_t FindCombinedSampler = OutString.find(SamplerName.c_str());
+				std::string SamplerNameDecl = SamplerName + ";";
+				size_t FindCombinedSampler = OutString.find(SamplerNameDecl.c_str());
 
 				if (FindCombinedSampler != std::string::npos)
 				{
