Fix GenerateGlslShader when a string is a substring of another in the  ReflectData

diff --git a/Engine/Shaders/Public/Platform/GL/GLCommon.ush b/Engine/Shaders/Public/Platform/GL/GLCommon.ush
--- a/Engine/Shaders/Public/Platform/GL/GLCommon.ush
+++ b/Engine/Shaders/Public/Platform/GL/GLCommon.ush
@@ -5,7 +5,7 @@
 =============================================================================*/
 
 // Update this GUID to improve shader recompilation for GL only shaders
-// GUID = 269FDCFB-1EE9-4A3F-B663-B99C3028EAE8
+// GUID = F1338239-A231-40B9-AAD4-421C7AD5306A
 
 #pragma once
 
diff --git a/Engine/Source/Developer/ShaderFormatOpenGL/Private/OpenGLShaderCompiler.cpp b/Engine/Source/Developer/ShaderFormatOpenGL/Private/OpenGLShaderCompiler.cpp
--- a/Engine/Source/Developer/ShaderFormatOpenGL/Private/OpenGLShaderCompiler.cpp
+++ b/Engine/Source/Developer/ShaderFormatOpenGL/Private/OpenGLShaderCompiler.cpp
@@ -2086,7 +2086,9 @@ static void ConvertToEmulatedUBs(std::string& GlslSource, const ReflectionData&
 
 						for (std::string const& SearchString : ReflectData.PackedUBArrays[PackedUBNamePair.Key])
 						{
-							if (!GlslSource.compare(UBVarPos, SearchString.length(), SearchString))
+							// This is done in case the current SearchString is a substring of another
+							std::string SearchStringWithBraces = SearchString + "[";
+							if (!GlslSource.compare(UBVarPos, SearchStringWithBraces.length(), SearchStringWithBraces))
 							{
 								GlslSource.replace(UBVarPos + SearchString.length(), 1, "(");
 
@@ -2409,7 +2411,9 @@ bool GenerateGlslShader(std::string& OutString, GLSLCompileParameters& GLSLCompi
 						OutString.replace(GlobalVarPos, GlobalVarString.length(), "_Globals_");
 						for (std::string const& SearchString : ReflectData.GlobalArrays)
 						{
-							if (!OutString.compare(GlobalVarPos, SearchString.length(), SearchString))
+							// This is done in case the current SearchString is a substring of another
+							std::string SearchStringWithBraces = SearchString + "[";
+							if (!OutString.compare(GlobalVarPos, SearchStringWithBraces.length(), SearchStringWithBraces))
 							{
 								OutString.replace(GlobalVarPos + SearchString.length(), 1, "(");
 
