Shader Compiler deprecation warning fix in ShaderFormatOpenGL.h.  Missed this one!

#rnx

diff --git a/Engine/Source/Developer/ShaderFormatOpenGL/Public/ShaderFormatOpenGL.h b/Engine/Source/Developer/ShaderFormatOpenGL/Public/ShaderFormatOpenGL.h
--- a/Engine/Source/Developer/ShaderFormatOpenGL/Public/ShaderFormatOpenGL.h
+++ b/Engine/Source/Developer/ShaderFormatOpenGL/Public/ShaderFormatOpenGL.h
@@ -72,7 +72,9 @@ protected:
 	virtual uint32 CalculateCrossCompilerFlags(GLSLVersion Version, const bool bFullPrecisionInPS, const FShaderCompilerFlags& CompilerFlags);
 
 	// set up compilation information like defines and HlslCompileTarget
+	PRAGMA_DISABLE_DEPRECATION_WARNINGS		// FShaderCompilerDefinitions will be made internal in the future, marked deprecated until then
 	virtual void SetupPerVersionCompilationEnvironment(GLSLVersion Version, class FShaderCompilerDefinitions& AdditionalDefines, EHlslCompileTarget& HlslCompilerTarget);
+	PRAGMA_ENABLE_DEPRECATION_WARNINGS
 
 	virtual void ConvertOpenGLVersionFromGLSLVersion(GLSLVersion InVersion, int& OutMajorVersion, int& OutMinorVersion);
 
