[shaders]
- enable minifier deadstripping for OpenGL shader format; bytecode is a 100% match with this change on multiple tested projects so it seems safe to enable
- deprecate FOpenGLFrontend; nothing in available code appears to inherit from it anymore so the API appears unnecessary. this makes the OpenGL shader format more similar to others (which do not expose compilation functions as API)

diff --git a/Engine/Source/Developer/ShaderFormatOpenGL/Private/OpenGLShaderCompiler.cpp b/Engine/Source/Developer/ShaderFormatOpenGL/Private/OpenGLShaderCompiler.cpp
--- a/Engine/Source/Developer/ShaderFormatOpenGL/Private/OpenGLShaderCompiler.cpp
+++ b/Engine/Source/Developer/ShaderFormatOpenGL/Private/OpenGLShaderCompiler.cpp
@@ -77,10 +77,6 @@ DEFINE_LOG_CATEGORY_STATIC(LogOpenGLShaderCompiler, Log, All);
 
 #define VALIDATE_GLSL_WITH_DRIVER		0
 #define ENABLE_IMAGINATION_COMPILER		1
-static FORCEINLINE bool IsPCESPlatform(GLSLVersion Version)
-{
-	return (Version == GLSL_150_ES3_1);
-}
 
 /*------------------------------------------------------------------------------
 	Shader compiling.
@@ -548,48 +544,40 @@ static uint32 ParseNumber(const TCHAR* Str)
 	return Num;
 }
 
-
-static ANSICHAR TranslateFrequencyToCrossCompilerPrefix(int32 Frequency)
+enum class EPlatformType
 {
-	switch (Frequency)
-	{
-		case SF_Vertex: return 'v';
-		case SF_Pixel: return 'p';
-		case SF_Geometry: return 'g';
-		case SF_Compute: return 'c';
-	}
-	return '\0';
-}
+	Android,
+	IOS,
+	Web,
+	Desktop
+};
 
-static TCHAR* SetIndex(TCHAR* Str, int32 Offset, int32 Index)
+struct FDeviceCapabilities
 {
-	check(Index >= 0 && Index < 100);
-
-	Str += Offset;
-	if (Index >= 10)
-	{
-		*Str++ = '0' + (TCHAR)(Index / 10);
-	}
-	*Str++ = '0' + (TCHAR)(Index % 10);
-	*Str = '\0';
-	return Str;
-}
-
+	EPlatformType TargetPlatform = EPlatformType::Android;
+};
 
+static bool PlatformSupportsOfflineCompilationInternal(const GLSLVersion ShaderVersion);
+static void FillDeviceCapsOfflineCompilationInternal(struct FDeviceCapabilities& Capabilities, const GLSLVersion ShaderVersion);
+static bool MoveHashLines(FString& Destination, FString &Source);
+static TSharedPtr<ANSICHAR> PrepareCodeForOfflineCompilationInternal(const GLSLVersion ShaderVersion, EShaderFrequency Frequency, const ANSICHAR* InShaderSource);
+static void PlatformCompileOfflineInternal(const FShaderCompilerInput& Input, FShaderCompilerOutput& ShaderOutput, const ANSICHAR* ShaderSource, const GLSLVersion ShaderVersion);
+static void CompileOfflineInternal(const FShaderCompilerInput& Input, FShaderCompilerOutput& Output, const GLSLVersion ShaderVersion, const ANSICHAR* InShaderSource);
 
 /**
  * Construct the final microcode from the compiled and verified shader source.
  * @param ShaderOutput - Where to store the microcode and parameter map.
+ * @param ShaderInput - The input struct for the shader being compiled.
  * @param InShaderSource - GLSL source with input/output signature.
  * @param SourceLen - The length of the GLSL source code.
+ * @param Version - The GLSL version to target.
  */
-void FOpenGLFrontend::BuildShaderOutput(
+void BuildShaderOutputInternal(
 	FShaderCompilerOutput& ShaderOutput,
 	const FShaderCompilerInput& ShaderInput,
 	const ANSICHAR* InShaderSource,
 	int32 SourceLen,
-	GLSLVersion Version
-	)
+	GLSLVersion Version)
 {
 	const ANSICHAR* USFSource = InShaderSource;
 	CrossCompiler::FHlslccHeader CCHeader;
@@ -604,7 +592,6 @@ void FOpenGLFrontend::BuildShaderOutput(
 	}
 
 	FOpenGLCodeHeader Header = {0};
-	FShaderParameterMap& ParameterMap = ShaderOutput.ParameterMap;
 	EShaderFrequency Frequency = (EShaderFrequency)ShaderOutput.Target.Frequency;
 
 	TBitArray<> UsedUniformBufferSlots;
@@ -681,10 +668,6 @@ void FOpenGLFrontend::BuildShaderOutput(
 		}
 	}
 
-	// general purpose binding name
-	TCHAR BindingName[] = TEXT("XYZ\0\0\0\0\0\0\0\0");
-	BindingName[0] = TranslateFrequencyToCrossCompilerPrefix(Frequency);
-
 	TMap<FString, FString> BindingNameMap;
 
 	// Then 'normal' uniform buffers.
@@ -693,18 +676,7 @@ void FOpenGLFrontend::BuildShaderOutput(
 		uint16 UBIndex = UniformBlock.Index;
 
 		UsedUniformBufferSlots[UBIndex] = true;
-		
-		if (OutputTrueParameterNames())
-		{
-			// make the final name this will be in the shader
-			BindingName[1] = 'b';
-			SetIndex(BindingName, 2, UBIndex);
-			BindingNameMap.Add(BindingName, UniformBlock.Name);
-		}
-		else
-		{
-			HandleReflectedUniformBuffer(UniformBlock.Name, UBIndex, ShaderOutput);
-		}
+		HandleReflectedUniformBuffer(UniformBlock.Name, UBIndex, ShaderOutput);
 		Header.Bindings.NumUniformBuffers++;
 	}
 
@@ -750,18 +722,8 @@ void FOpenGLFrontend::BuildShaderOutput(
 	TMap<int, TMap<ANSICHAR, uint16> > PackedUniformBuffersSize;
 	for (auto& PackedUB : CCHeader.PackedUBs)
 	{
-		checkf(OutputTrueParameterNames() == false, TEXT("Unexpected Packed UBs used with a shader format that needs true parameter names - If this is hit, we need to figure out how to handle them"));
-
 		UsedUniformBufferSlots[PackedUB.Attribute.Index] = true;
-		if (OutputTrueParameterNames())
-		{
-			BindingName[1] = 'b';
-			// ???
-		}
-		else
-		{
-			HandleReflectedUniformBuffer(PackedUB.Attribute.Name, PackedUB.Attribute.Index, ShaderOutput);
-		}
+		HandleReflectedUniformBuffer(PackedUB.Attribute.Name, PackedUB.Attribute.Index, ShaderOutput);
 		Header.Bindings.NumUniformBuffers++;
 
 		// Nothing else...
@@ -839,7 +801,6 @@ void FOpenGLFrontend::BuildShaderOutput(
 	Header.Bindings.PackedUniformBuffers.Reserve(PackedUniformBuffersSize.Num());
 	for (auto Iterator = PackedUniformBuffersSize.CreateIterator(); Iterator; ++Iterator)
 	{
-		int BufferIndex = Iterator.Key();
 		auto& ArraySizes = Iterator.Value();
 		TArray<CrossCompiler::FPackedArrayInfo> InfoArray;
 		InfoArray.Reserve(ArraySizes.Num());
@@ -867,16 +828,7 @@ void FOpenGLFrontend::BuildShaderOutput(
 	// Then samplers.
 	for (auto& Sampler : CCHeader.Samplers)
 	{
-		if (OutputTrueParameterNames())
-		{
-			BindingName[1] = 's';
-			SetIndex(BindingName, 2, Sampler.Offset);
-			BindingNameMap.Add(BindingName, Sampler.Name);
-		}
-		else
-		{
-			HandleReflectedShaderResource(Sampler.Name, Sampler.Offset, Sampler.Count, ShaderOutput);
-		}
+		HandleReflectedShaderResource(Sampler.Name, Sampler.Offset, Sampler.Count, ShaderOutput);
 
 		Header.Bindings.NumSamplers = FMath::Max<uint8>(
 			Header.Bindings.NumSamplers,
@@ -885,32 +837,14 @@ void FOpenGLFrontend::BuildShaderOutput(
 
 		for (auto& SamplerState : Sampler.SamplerStates)
 		{
-			if (OutputTrueParameterNames())
-			{
-				// add an entry for the sampler parameter as well
-				BindingNameMap.Add(FString(BindingName) + TEXT("_samp"), SamplerState);
-			}
-			else
-			{
-				HandleReflectedShaderSampler(SamplerState, Sampler.Offset, Sampler.Count, ShaderOutput);
-			}
+			HandleReflectedShaderSampler(SamplerState, Sampler.Offset, Sampler.Count, ShaderOutput);
 		}
 	}
 
 	// Then UAVs (images in GLSL)
 	for (auto& UAV : CCHeader.UAVs)
 	{
-		if (OutputTrueParameterNames())
-		{
-			// make the final name this will be in the shader
-			BindingName[1] = 'i';
-			SetIndex(BindingName, 2, UAV.Offset);
-			BindingNameMap.Add(BindingName, UAV.Name);
-		}
-		else
-		{
-			HandleReflectedShaderUAV(UAV.Name, UAV.Offset, UAV.Count, ShaderOutput);
-		}
+		HandleReflectedShaderUAV(UAV.Name, UAV.Offset, UAV.Count, ShaderOutput);
 
 		Header.Bindings.NumUAVs = FMath::Max<uint8>(
 			Header.Bindings.NumSamplers,
@@ -919,9 +853,7 @@ void FOpenGLFrontend::BuildShaderOutput(
 	}
 
 	Header.ShaderName = CCHeader.Name;
-
-	// perform any post processing this frontend class may need to do
-	ShaderOutput.bSucceeded = PostProcessShaderSource(Version, Frequency, USFSource, SourceLen + 1 - (USFSource - InShaderSource), ParameterMap, BindingNameMap, ShaderOutput.Errors, ShaderInput);
+	ShaderOutput.bSucceeded = true;
 
 	// Build the SRT for this shader.
 	{
@@ -942,7 +874,7 @@ void FOpenGLFrontend::BuildShaderOutput(
 		BuildResourceTableTokenStream(GenericSRT.UnorderedAccessViewMap, GenericSRT.MaxBoundResourceTable, Header.Bindings.ShaderResourceTable.UnorderedAccessViewMap);
 	}
 
-	const int32 MaxSamplers = GetMaxSamplers(Version);
+	constexpr int32 MaxSamplers = 16;
 
 	if (Header.Bindings.NumSamplers > MaxSamplers)
 	{
@@ -958,11 +890,7 @@ void FOpenGLFrontend::BuildShaderOutput(
 		FMemoryWriter Ar(ShaderOutput.ShaderCode.GetWriteAccess(), true);
 		Ar << Header;
 
-		if (OptionalSerializeOutputAndReturnIfSerialized(Ar) == false)
-		{
-			Ar.Serialize((void*)USFSource, SourceLen + 1 - (USFSource - InShaderSource));
-			ShaderOutput.bSucceeded = true;
-		}
+		Ar.Serialize((void*)USFSource, SourceLen + 1 - (USFSource - InShaderSource));
 
 		// extract final source code as requested by the Material Editor
 		if (ShaderInput.ExtraSettings.bExtractShaderSource)
@@ -980,7 +908,7 @@ void FOpenGLFrontend::BuildShaderOutput(
 		// if available, attempt run an offline compilation and extract statistics
 		if (ShaderInput.ExtraSettings.OfflineCompilerPath.Len() > 0)
 		{
-			CompileOffline(ShaderInput, ShaderOutput, Version, USFSource);
+			CompileOfflineInternal(ShaderInput, ShaderOutput, Version, USFSource);
 		}
 		else
 		{
@@ -991,7 +919,17 @@ void FOpenGLFrontend::BuildShaderOutput(
 	}
 }
 
-void FOpenGLFrontend::ConvertOpenGLVersionFromGLSLVersion(GLSLVersion InVersion, int& OutMajorVersion, int& OutMinorVersion)
+void FOpenGLFrontend::BuildShaderOutput(
+	FShaderCompilerOutput& ShaderOutput,
+	const FShaderCompilerInput& ShaderInput,
+	const ANSICHAR* InShaderSource,
+	int32 SourceLen,
+	GLSLVersion Version)
+{
+	BuildShaderOutputInternal(ShaderOutput, ShaderInput, InShaderSource, SourceLen, Version);
+}
+
+static void ConvertOpenGLVersionFromGLSLVersionInternal(GLSLVersion InVersion, int& OutMajorVersion, int& OutMinorVersion)
 {
 	switch(InVersion)
 	{
@@ -1012,13 +950,20 @@ void FOpenGLFrontend::ConvertOpenGLVersionFromGLSLVersion(GLSLVersion InVersion,
 	}
 }
 
+void FOpenGLFrontend::ConvertOpenGLVersionFromGLSLVersion(GLSLVersion InVersion, int& OutMajorVersion, int& OutMinorVersion)
+{
+	return ConvertOpenGLVersionFromGLSLVersionInternal(InVersion, OutMajorVersion, OutMinorVersion);
+}
+
 /**
  * Precompile a GLSL shader.
  * @param ShaderOutput - The precompiled shader.
  * @param ShaderInput - The shader input.
- * @param InPreprocessedShader - The preprocessed source code.
+ * @param ShaderSource - The preprocessed source code.
+ * @param Version - The GLSL language version to target.
+ * @param Frequency - The shader stage
  */
-void FOpenGLFrontend::PrecompileShader(FShaderCompilerOutput& ShaderOutput, const FShaderCompilerInput& ShaderInput, const ANSICHAR* ShaderSource, GLSLVersion Version, EHlslShaderFrequency Frequency)
+static void PrecompileShaderInternal(FShaderCompilerOutput& ShaderOutput, const FShaderCompilerInput& ShaderInput, const ANSICHAR* ShaderSource, GLSLVersion Version, EHlslShaderFrequency Frequency)
 {
 	check(ShaderInput.Target.Frequency < SF_NumFrequencies);
 
@@ -1038,7 +983,7 @@ void FOpenGLFrontend::PrecompileShader(FShaderCompilerOutput& ShaderOutput, cons
 	void* PrevContextPtr;
 	int MajorVersion = 0;
 	int MinorVersion = 0;
-	ConvertOpenGLVersionFromGLSLVersion(Version, MajorVersion, MinorVersion);
+	ConvertOpenGLVersionFromGLSLVersionInternal(Version, MajorVersion, MinorVersion);
 	PlatformInitOpenGL(ContextPtr, PrevContextPtr, MajorVersion, MinorVersion);
 
 	GLint SourceLen = FCStringAnsi::Strlen(ShaderSource);
@@ -1056,7 +1001,7 @@ void FOpenGLFrontend::PrecompileShader(FShaderCompilerOutput& ShaderOutput, cons
 		if (CompileStatus == GL_TRUE)
 		{
 			ShaderOutput.Target = ShaderInput.Target;
-			BuildShaderOutput(
+			BuildShaderOutputInternal(
 				ShaderOutput,
 				ShaderInput,
 				ShaderSource,
@@ -1108,10 +1053,14 @@ void FOpenGLFrontend::PrecompileShader(FShaderCompilerOutput& ShaderOutput, cons
 	PlatformReleaseOpenGL(ContextPtr, PrevContextPtr);
 }
 
+void FOpenGLFrontend::PrecompileShader(FShaderCompilerOutput& ShaderOutput, const FShaderCompilerInput& ShaderInput, const ANSICHAR* ShaderSource, GLSLVersion Version, EHlslShaderFrequency Frequency)
+{
+	PrecompileShaderInternal(ShaderOutput, ShaderInput, ShaderSource, Version, Frequency);
+}
 
 PRAGMA_DISABLE_DEPRECATION_WARNINGS		// FShaderCompilerDefinitions will be made internal in the future, marked deprecated until then
 
-void FOpenGLFrontend::SetupPerVersionCompilationEnvironment(GLSLVersion Version, FShaderCompilerDefinitions& AdditionalDefines, EHlslCompileTarget& HlslCompilerTarget)
+static void SetupPerVersionCompilationEnvironmentInternal(GLSLVersion Version, FShaderCompilerDefinitions& AdditionalDefines, EHlslCompileTarget& HlslCompilerTarget)
 {
 	switch (Version)
 	{
@@ -1135,6 +1084,11 @@ void FOpenGLFrontend::SetupPerVersionCompilationEnvironment(GLSLVersion Version,
 	AdditionalDefines.SetDefine(TEXT("OPENGL_PROFILE"), 1);
 }
 
+void FOpenGLFrontend::SetupPerVersionCompilationEnvironment(GLSLVersion Version, FShaderCompilerDefinitions& AdditionalDefines, EHlslCompileTarget& HlslCompilerTarget)
+{
+	SetupPerVersionCompilationEnvironmentInternal(Version, AdditionalDefines, HlslCompilerTarget);
+}
+
 PRAGMA_ENABLE_DEPRECATION_WARNINGS
 
 uint32 FOpenGLFrontend::GetMaxSamplers(GLSLVersion Version)
@@ -1142,7 +1096,7 @@ uint32 FOpenGLFrontend::GetMaxSamplers(GLSLVersion Version)
 	return 16;
 }
 
-uint32 FOpenGLFrontend::CalculateCrossCompilerFlags(GLSLVersion Version, const bool bFullPrecisionInPS, const FShaderCompilerFlags& CompilerFlags)
+static uint32 CalculateCrossCompilerFlagsInternal(GLSLVersion Version, const bool bFullPrecisionInPS, const FShaderCompilerFlags& CompilerFlags)
 {
 	uint32  CCFlags = HLSLCC_NoPreprocess | HLSLCC_PackUniforms | HLSLCC_DX11ClipSpace | HLSLCC_RetainSizes;
 
@@ -1167,11 +1121,22 @@ uint32 FOpenGLFrontend::CalculateCrossCompilerFlags(GLSLVersion Version, const b
 	return CCFlags;
 }
 
-FGlslCodeBackend* FOpenGLFrontend::CreateBackend(GLSLVersion Version, uint32 CCFlags, EHlslCompileTarget HlslCompilerTarget)
+uint32 FOpenGLFrontend::CalculateCrossCompilerFlags(GLSLVersion Version, const bool bFullPrecisionInPS, const FShaderCompilerFlags& CompilerFlags)
+{
+	return CalculateCrossCompilerFlagsInternal(Version, bFullPrecisionInPS, CompilerFlags);
+}
+
+FGlslCodeBackend* CreateBackendInternal(GLSLVersion Version, uint32 CCFlags, EHlslCompileTarget HlslCompilerTarget)
 {
 	return new FGlslCodeBackend(CCFlags, HlslCompilerTarget);
 }
 
+
+FGlslCodeBackend* FOpenGLFrontend::CreateBackend(GLSLVersion Version, uint32 CCFlags, EHlslCompileTarget HlslCompilerTarget)
+{
+	return CreateBackendInternal(Version, CCFlags, HlslCompilerTarget);
+}
+
 class FGlsl430LanguageSpec : public FGlslLanguageSpec
 {
 public:
@@ -1181,11 +1146,16 @@ public:
 	virtual bool EmulateStructuredWithTypedBuffers() const override { return false; }
 };
 
-FGlslLanguageSpec* FOpenGLFrontend::CreateLanguageSpec(GLSLVersion Version, bool bDefaultPrecisionIsHalf)
+FGlslLanguageSpec* CreateLanguageSpecInternal(GLSLVersion Version, bool bDefaultPrecisionIsHalf)
 {
 	return new FGlslLanguageSpec(bDefaultPrecisionIsHalf);
 }
 
+FGlslLanguageSpec* FOpenGLFrontend::CreateLanguageSpec(GLSLVersion Version, bool bDefaultPrecisionIsHalf)
+{
+	return CreateLanguageSpecInternal(Version, bDefaultPrecisionIsHalf);
+}
+
 #if DXC_SUPPORTED
 
 static const ANSICHAR* GetFrequencyPrefix(EShaderFrequency Frequency)
@@ -1213,10 +1183,6 @@ struct PackedUBMemberInfo
 
 void WritePackedUBHeader(CrossCompiler::FHlslccHeaderWriter& CCHeaderWriter, const TMap<uint32_t, TArray<PackedUBMemberInfo>>& UBMemberInfo, const TMap<uint32, std::string>& UBNames)
 {
-	bool bFirst = true;
-
-	bool bNeedsHeader = true;
-
 	for (const auto & Pair : UBNames)
 	{
 		FString Name(Pair.Value.c_str());
@@ -1276,7 +1242,6 @@ void WritePackedUBHeader(CrossCompiler::FHlslccHeaderWriter& CCHeaderWriter, con
 void GetSpvVarQualifier(const SpvReflectBlockVariable& Member, FString & Out)
 {
 	auto const type = *Member.type_description;
-	const uint32 MbrSize = Member.size / sizeof(float);
 
 	FString TypeQualifier;
 
@@ -1546,7 +1511,6 @@ void ParseReflectionData(const FShaderCompilerInput& ShaderInput, CrossCompiler:
 	uint32 TextureIndices = 0xffffffff;
 	uint32 UBOIndices = 0xffffffff;
 	uint32 SamplerIndices = 0xffffffff;
-	uint32_t PackedUBIndex = 0;
 
 	for (auto const& Binding : ReflectionBindings.TBufferUAVs)
 	{
@@ -3119,8 +3083,6 @@ static bool CompileToGlslWithShaderConductor(
 
 		ParseReflectionData(Input, CCHeaderWriter, ReflectData, SpirvData, Reflection, SPIRV_DummySamplerName, Frequency, bEmulatedUBs);
 		
-		const ANSICHAR* FrequencyPrefix = GetFrequencyPrefix(Frequency);
-		
 		// Step 2 : End of reflection
 		// 
 		// Overwrite updated SPIRV code
@@ -3273,37 +3235,12 @@ static bool CompileToGlslWithShaderConductor(
 
 #endif // DXC_SUPPORTED
 
-
-static inline FString GetExtension(EHlslShaderFrequency Frequency, bool bAddDot = true)
-{
-	const TCHAR* Name = nullptr;
-	switch (Frequency)
-	{
-	default:
-		check(0);
-		// fallthrough...
-
-	case HSF_PixelShader:		Name = TEXT(".frag"); break;
-	case HSF_VertexShader:		Name = TEXT(".vert"); break;
-	case HSF_ComputeShader:		Name = TEXT(".comp"); break;
-	case HSF_GeometryShader:	Name = TEXT(".geom"); break;
-	case HSF_HullShader:		Name = TEXT(".tesc"); break;
-	case HSF_DomainShader:		Name = TEXT(".tese"); break;
-	}
-
-	if (!bAddDot)
-	{
-		++Name;
-	}
-	return FString(Name);
-}
-
 /**
  * Compile a shader for OpenGL on Windows.
  * @param Input - The input shader code and environment.
  * @param Output - Contains shader compilation results upon return.
  */
-void FOpenGLFrontend::CompileShader(const FShaderCompilerInput& Input, FShaderCompilerOutput& Output, const FString& WorkingDirectory, GLSLVersion Version)
+void CompileOpenGLShader(const FShaderCompilerInput& Input, FShaderCompilerOutput& Output, const FString& WorkingDirectory, GLSLVersion Version)
 {
 	PRAGMA_DISABLE_DEPRECATION_WARNINGS		// FShaderCompilerDefinitions will be made internal in the future, marked deprecated until then
 
@@ -3313,7 +3250,7 @@ void FOpenGLFrontend::CompileShader(const FShaderCompilerInput& Input, FShaderCo
 	ECompilerFlags PlatformFlowControl = CFLAG_AvoidFlowControl;
 
 	// set up compiler env based on version
-	SetupPerVersionCompilationEnvironment(Version, AdditionalDefines, HlslCompilerTarget);
+	SetupPerVersionCompilationEnvironmentInternal(Version, AdditionalDefines, HlslCompilerTarget);
 
 #if DXC_SUPPORTED
 	const bool bUseSC = Input.Environment.CompilerFlags.Contains(CFLAG_ForceDXC);
@@ -3362,7 +3299,6 @@ void FOpenGLFrontend::CompileShader(const FShaderCompilerInput& Input, FShaderCo
 
 	char* GlslShaderSource = NULL;
 	char* ErrorLog = NULL;
-	const bool bIsSM5 = IsSM5(Version);
 
 	const EHlslShaderFrequency FrequencyTable[] =
 	{
@@ -3410,7 +3346,7 @@ void FOpenGLFrontend::CompileShader(const FShaderCompilerInput& Input, FShaderCo
 	}
 #endif // UE_OPENGL_SHADER_COMPILER_ALLOW_DEAD_CODE_REMOVAL
 
-	uint32 CCFlags = CalculateCrossCompilerFlags(Version, Input.Environment.FullPrecisionInPS, Input.Environment.CompilerFlags);
+	uint32 CCFlags = CalculateCrossCompilerFlagsInternal(Version, Input.Environment.FullPrecisionInPS, Input.Environment.CompilerFlags);
 
 	// Required as we added the RemoveUniformBuffersFromSource() function (the cross-compiler won't be able to interpret comments w/o a preprocessor)
 	CCFlags &= ~HLSLCC_NoPreprocess;
@@ -3426,12 +3362,12 @@ void FOpenGLFrontend::CompileShader(const FShaderCompilerInput& Input, FShaderCo
 #endif // DXC_SUPPORTED
 	{
 		UE::ShaderCompilerCommon::DumpDebugShaderData(Input, PreprocessedShader);
-		
+
 		CCFlags |= HLSLCC_NoValidation;
-		FGlslCodeBackend* BackEnd = CreateBackend(Version, CCFlags, HlslCompilerTarget);
+		FGlslCodeBackend* BackEnd = CreateBackendInternal(Version, CCFlags, HlslCompilerTarget);
 
 		bool bDefaultPrecisionIsHalf = (CCFlags & HLSLCC_UseFullPrecisionInPS) == 0;
-		FGlslLanguageSpec* LanguageSpec = CreateLanguageSpec(Version, bDefaultPrecisionIsHalf);
+		FGlslLanguageSpec* LanguageSpec = CreateLanguageSpecInternal(Version, bDefaultPrecisionIsHalf);
 
 		{
 			FScopeLock HlslCcLock(CrossCompiler::GetCrossCompilerLock());
@@ -3466,11 +3402,11 @@ void FOpenGLFrontend::CompileShader(const FShaderCompilerInput& Input, FShaderCo
 		}
 
 #if VALIDATE_GLSL_WITH_DRIVER
-		PrecompileShader(Output, Input, GlslShaderSource, Version, HlslFrequency);
+		PrecompileShaderInternal(Output, Input, GlslShaderSource, Version, HlslFrequency);
 #else // VALIDATE_GLSL_WITH_DRIVER
 		int32 SourceLen = FCStringAnsi::Strlen(GlslShaderSource); //-V595
 		Output.Target = Input.Target;
-		BuildShaderOutput(Output, Input, GlslShaderSource, SourceLen, Version);
+		BuildShaderOutputInternal(Output, Input, GlslShaderSource, SourceLen, Version);
 #endif // VALIDATE_GLSL_WITH_DRIVER
 
 		if (bDumpDebugInfo && GlslShaderSource != nullptr)
@@ -3512,20 +3448,13 @@ void FOpenGLFrontend::CompileShader(const FShaderCompilerInput& Input, FShaderCo
 	//ShaderParameterParser.ValidateShaderParameterTypes(Input, Output);
 }
 
-enum class EPlatformType
-{
-	Android,
-	IOS,
-	Web,
-	Desktop
-};
 
-struct FDeviceCapabilities
+void FOpenGLFrontend::CompileShader(const FShaderCompilerInput& Input, FShaderCompilerOutput& Output, const FString& WorkingDirectory, GLSLVersion Version)
 {
-	EPlatformType TargetPlatform = EPlatformType::Android;
-};
+	CompileOpenGLShader(Input, Output, WorkingDirectory, Version);
+}
 
-void FOpenGLFrontend::FillDeviceCapsOfflineCompilation(struct FDeviceCapabilities& Capabilities, const GLSLVersion ShaderVersion) const
+static void FillDeviceCapsOfflineCompilationInternal(struct FDeviceCapabilities& Capabilities, const GLSLVersion ShaderVersion)
 {
 	FMemory::Memzero(Capabilities);
 
@@ -3539,6 +3468,11 @@ void FOpenGLFrontend::FillDeviceCapsOfflineCompilation(struct FDeviceCapabilitie
 	}
 }
 
+void FOpenGLFrontend::FillDeviceCapsOfflineCompilation(struct FDeviceCapabilities& Capabilities, const GLSLVersion ShaderVersion) const
+{
+	FillDeviceCapsOfflineCompilationInternal(Capabilities, ShaderVersion);
+}
+
 static bool MoveHashLines(FString& Destination, FString &Source)
 {
 	int32 Index = 0;
@@ -3576,13 +3510,13 @@ static bool MoveHashLines(FString& Destination, FString &Source)
 	return bFound;
 }
 
-TSharedPtr<ANSICHAR> FOpenGLFrontend::PrepareCodeForOfflineCompilation(const GLSLVersion ShaderVersion, EShaderFrequency Frequency, const ANSICHAR* InShaderSource) const
+static TSharedPtr<ANSICHAR> PrepareCodeForOfflineCompilationInternal(const GLSLVersion ShaderVersion, EShaderFrequency Frequency, const ANSICHAR* InShaderSource)
 {
 	FString OriginalShaderSource(ANSI_TO_TCHAR(InShaderSource));
 	FString StrOutSource;
 
 	FDeviceCapabilities Capabilities;
-	FillDeviceCapsOfflineCompilation(Capabilities, ShaderVersion);
+	FillDeviceCapsOfflineCompilationInternal(Capabilities, ShaderVersion);
 
 	// Whether we need to emit mobile multi-view code or not.
 	const bool bEmitMobileMultiView = OriginalShaderSource.Find(TEXT("gl_ViewID_OVR")) != INDEX_NONE;
@@ -3590,7 +3524,6 @@ TSharedPtr<ANSICHAR> FOpenGLFrontend::PrepareCodeForOfflineCompilation(const GLS
 	// Whether we need to emit texture external code or not.
 	const bool bEmitTextureExternal = OriginalShaderSource.Find(TEXT("samplerExternalOES")) != INDEX_NONE;
 
-	bool bNeedsExtDrawInstancedDefine = false;
 	if (Capabilities.TargetPlatform == EPlatformType::Android || Capabilities.TargetPlatform == EPlatformType::Web)
 	{
 		const TCHAR *ES320Version = TEXT("#version 320 es");
@@ -3599,10 +3532,8 @@ TSharedPtr<ANSICHAR> FOpenGLFrontend::PrepareCodeForOfflineCompilation(const GLS
 		OriginalShaderSource.RemoveFromStart(ES320Version);
 	}
 
-	const GLenum TypeEnum = GLFrequencyTable[Frequency];
 	// The incoming glsl may have preprocessor code that is dependent on defines introduced via the engine.
 	// This is the place to insert such engine preprocessor defines, immediately after the glsl version declarati
-
 	if (bEmitMobileMultiView)
 	{
 		MoveHashLines(StrOutSource, OriginalShaderSource);
@@ -3657,43 +3588,49 @@ TSharedPtr<ANSICHAR> FOpenGLFrontend::PrepareCodeForOfflineCompilation(const GLS
 	return RetShaderSource;
 }
 
-bool FOpenGLFrontend::PlatformSupportsOfflineCompilation(const GLSLVersion ShaderVersion) const
+TSharedPtr<ANSICHAR> FOpenGLFrontend::PrepareCodeForOfflineCompilation(const GLSLVersion ShaderVersion, EShaderFrequency Frequency, const ANSICHAR* InShaderSource) const
 {
-	switch (ShaderVersion)
-	{
-		// desktop
-		case GLSL_150_ES3_1:
-		// switch
-		case GLSL_SWITCH:
-		case GLSL_SWITCH_FORWARD:
-			return false;
-		break;
-		case GLSL_ES3_1_ANDROID:
-			return true;
-		break;
-	}
+	return PrepareCodeForOfflineCompilationInternal(ShaderVersion, Frequency, InShaderSource);
+}
+
+static bool PlatformSupportsOfflineCompilationInternal(const GLSLVersion ShaderVersion)
+{
+	return ShaderVersion == GLSL_ES3_1_ANDROID;
+}
 
-	return false;
+bool FOpenGLFrontend::PlatformSupportsOfflineCompilation(const GLSLVersion ShaderVersion) const
+{
+	return PlatformSupportsOfflineCompilationInternal(ShaderVersion);
 }
 
-void FOpenGLFrontend::CompileOffline(const FShaderCompilerInput& Input, FShaderCompilerOutput& Output, const GLSLVersion ShaderVersion, const ANSICHAR* InShaderSource)
+static void CompileOfflineInternal(const FShaderCompilerInput& Input, FShaderCompilerOutput& Output, const GLSLVersion ShaderVersion, const ANSICHAR* InShaderSource)
 {
-	const bool bSupportsOfflineCompilation = PlatformSupportsOfflineCompilation(ShaderVersion);
+	const bool bSupportsOfflineCompilation = PlatformSupportsOfflineCompilationInternal(ShaderVersion);
 
 	if (!bSupportsOfflineCompilation)
 	{
 		return;
 	}
 
-	TSharedPtr<ANSICHAR> ShaderSource = PrepareCodeForOfflineCompilation(ShaderVersion, (EShaderFrequency)Input.Target.Frequency, InShaderSource);
+	TSharedPtr<ANSICHAR> ShaderSource = PrepareCodeForOfflineCompilationInternal(ShaderVersion, (EShaderFrequency)Input.Target.Frequency, InShaderSource);
 
-	PlatformCompileOffline(Input, Output, ShaderSource.Get(), ShaderVersion);
+	PlatformCompileOfflineInternal(Input, Output, ShaderSource.Get(), ShaderVersion);
 }
 
-void FOpenGLFrontend::PlatformCompileOffline(const FShaderCompilerInput& Input, FShaderCompilerOutput& ShaderOutput, const ANSICHAR* ShaderSource, const GLSLVersion ShaderVersion)
+void FOpenGLFrontend::CompileOffline(const FShaderCompilerInput& Input, FShaderCompilerOutput& Output, const GLSLVersion ShaderVersion, const ANSICHAR* InShaderSource)
+{
+	CompileOfflineInternal(Input, Output, ShaderVersion, InShaderSource);
+}
+
+static void PlatformCompileOfflineInternal(const FShaderCompilerInput& Input, FShaderCompilerOutput& ShaderOutput, const ANSICHAR* ShaderSource, const GLSLVersion ShaderVersion)
 {
 	if (ShaderVersion == GLSL_ES3_1_ANDROID)
 	{
 		CompileOfflineMali(Input, ShaderOutput, ShaderSource, FPlatformString::Strlen(ShaderSource), false);
 	}
 }
+
+void FOpenGLFrontend::PlatformCompileOffline(const FShaderCompilerInput& Input, FShaderCompilerOutput& ShaderOutput, const ANSICHAR* ShaderSource, const GLSLVersion ShaderVersion)
+{
+	PlatformCompileOfflineInternal(Input, ShaderOutput, ShaderSource, ShaderVersion);
+}
diff --git a/Engine/Source/Developer/ShaderFormatOpenGL/Private/ShaderFormatOpenGL.cpp b/Engine/Source/Developer/ShaderFormatOpenGL/Private/ShaderFormatOpenGL.cpp
--- a/Engine/Source/Developer/ShaderFormatOpenGL/Private/ShaderFormatOpenGL.cpp
+++ b/Engine/Source/Developer/ShaderFormatOpenGL/Private/ShaderFormatOpenGL.cpp
@@ -11,7 +11,9 @@
 
 static FName NAME_GLSL_150_ES3_1(TEXT("GLSL_150_ES31"));
 static FName NAME_GLSL_ES3_1_ANDROID(TEXT("GLSL_ES3_1_ANDROID"));
- 
+
+extern void CompileOpenGLShader(const FShaderCompilerInput& Input, FShaderCompilerOutput& Output, const FString& WorkingDirectory, GLSLVersion Version);
+
 class FShaderFormatGLSL : public IShaderFormat
 {
 	enum
@@ -76,9 +78,7 @@ public:
 
 		GLSLVersion Version = TranslateFormatNameToEnum(Format);
 
-		FOpenGLFrontend Frontend;
-		// the frontend will run the cross compiler
-		Frontend.CompileShader(Input, Output, WorkingDirectory, Version);
+		CompileOpenGLShader(Input, Output, WorkingDirectory, Version);
 	}
 
 	virtual const TCHAR* GetPlatformIncludeDirectory() const
diff --git a/Engine/Source/Developer/ShaderFormatOpenGL/Public/ShaderFormatOpenGL.h b/Engine/Source/Developer/ShaderFormatOpenGL/Public/ShaderFormatOpenGL.h
--- a/Engine/Source/Developer/ShaderFormatOpenGL/Public/ShaderFormatOpenGL.h
+++ b/Engine/Source/Developer/ShaderFormatOpenGL/Public/ShaderFormatOpenGL.h
@@ -19,7 +19,7 @@ class FArchive;
 
 // Controls whether r.Shaders.RemoveDeadCode should be honored
 #ifndef UE_OPENGL_SHADER_COMPILER_ALLOW_DEAD_CODE_REMOVAL
-#define UE_OPENGL_SHADER_COMPILER_ALLOW_DEAD_CODE_REMOVAL 0 // Disabled initially, requires testing
+#define UE_OPENGL_SHADER_COMPILER_ALLOW_DEAD_CODE_REMOVAL 1
 #endif // UE_OPENGL_SHADER_COMPILER_ALLOW_DEAD_CODE_REMOVAL
 
 class FShaderCompilerFlags;
@@ -47,45 +47,56 @@ enum GLSLVersion
 class SHADERFORMATOPENGL_API FOpenGLFrontend
 {
 public:
+	UE_DEPRECATED(5.4, "FOpenGLFrontend is deprecated and derived implementations no longer possible; OpenGL shader compilation should only be triggered from within the shader format code.")
 	virtual ~FOpenGLFrontend()
 	{
 
 	}
 
+	UE_DEPRECATED(5.4, "FOpenGLFrontend is deprecated and derived implementations no longer possible; OpenGL shader compilation should only be triggered from within the shader format code.")
 	void CompileShader(const struct FShaderCompilerInput& Input, struct FShaderCompilerOutput& Output, const class FString& WorkingDirectory, GLSLVersion Version);
 
 protected:
 	// if true, the shader output map will contain true names (i.e. ColorModifier) instead of helper names for runtime binding (i.e. pb_5)
+	UE_DEPRECATED(5.4, "FOpenGLFrontend is deprecated and derived implementations no longer possible; OpenGL shader compilation should only be triggered from within the shader format code.")
 	virtual bool OutputTrueParameterNames()
 	{
 		return false;
 	}
 
+	UE_DEPRECATED(5.4, "FOpenGLFrontend is deprecated and derived implementations no longer possible; OpenGL shader compilation should only be triggered from within the shader format code.")
 	virtual bool IsSM5(GLSLVersion Version)
 	{
 		return false;
 	}
 
 	// what is the max number of samplers the shader platform can use?
+	UE_DEPRECATED(5.4, "FOpenGLFrontend is deprecated and derived implementations no longer possible; OpenGL shader compilation should only be triggered from within the shader format code.")
 	virtual uint32 GetMaxSamplers(GLSLVersion Version);
 
+	UE_DEPRECATED(5.4, "FOpenGLFrontend is deprecated and derived implementations no longer possible; OpenGL shader compilation should only be triggered from within the shader format code.")
 	virtual uint32 CalculateCrossCompilerFlags(GLSLVersion Version, const bool bFullPrecisionInPS, const FShaderCompilerFlags& CompilerFlags);
 
 	// set up compilation information like defines and HlslCompileTarget
 	PRAGMA_DISABLE_DEPRECATION_WARNINGS		// FShaderCompilerDefinitions will be made internal in the future, marked deprecated until then
+	UE_DEPRECATED(5.4, "FOpenGLFrontend is deprecated and derived implementations no longer possible; OpenGL shader compilation should only be triggered from within the shader format code.")
 	virtual void SetupPerVersionCompilationEnvironment(GLSLVersion Version, class FShaderCompilerDefinitions& AdditionalDefines, EHlslCompileTarget& HlslCompilerTarget);
 	PRAGMA_ENABLE_DEPRECATION_WARNINGS
 
+	UE_DEPRECATED(5.4, "FOpenGLFrontend is deprecated and derived implementations no longer possible; OpenGL shader compilation should only be triggered from within the shader format code.")
 	virtual void ConvertOpenGLVersionFromGLSLVersion(GLSLVersion InVersion, int& OutMajorVersion, int& OutMinorVersion);
 
 	// create the compiling backend
+	UE_DEPRECATED(5.4, "FOpenGLFrontend is deprecated and derived implementations no longer possible; OpenGL shader compilation should only be triggered from within the shader format code.")
 	virtual struct FGlslCodeBackend* CreateBackend(GLSLVersion Version, uint32 CCFlags, EHlslCompileTarget HlslCompilerTarget);
 
 	// create the language spec
+	UE_DEPRECATED(5.4, "FOpenGLFrontend is deprecated and derived implementations no longer possible; OpenGL shader compilation should only be triggered from within the shader format code.")
 	virtual class FGlslLanguageSpec* CreateLanguageSpec(GLSLVersion Version, bool bDefaultPrecisionIsHalf);
 
 
 	// Allow a subclass to perform additional work on the cross compiled source code
+	UE_DEPRECATED(5.4, "FOpenGLFrontend is deprecated and derived implementations no longer possible; OpenGL shader compilation should only be triggered from within the shader format code.")
 	virtual bool PostProcessShaderSource(GLSLVersion Version, EShaderFrequency Frequency, const ANSICHAR* ShaderSource,
 		uint32 SourceLen, class FShaderParameterMap& ParameterMap, TMap<FString, FString>& BindingNameMap, TArray<struct FShaderCompilerError>& Errors,
 		const FShaderCompilerInput& ShaderInput)
@@ -94,24 +105,32 @@ protected:
 	}
 
 	// allow subclass to write out different output, returning true if it did write everything it needed
+	UE_DEPRECATED(5.4, "FOpenGLFrontend is deprecated and derived implementations no longer possible; OpenGL shader compilation should only be triggered from within the shader format code.")
 	virtual bool OptionalSerializeOutputAndReturnIfSerialized(FArchive& Ar)
 	{
 		return false;
 	}
 
 
+	UE_DEPRECATED(5.4, "FOpenGLFrontend is deprecated and derived implementations no longer possible; OpenGL shader compilation should only be triggered from within the shader format code.")
 	void BuildShaderOutput(FShaderCompilerOutput& ShaderOutput, const FShaderCompilerInput& ShaderInput, const ANSICHAR* InShaderSource, int32 SourceLen, GLSLVersion Version);
+	UE_DEPRECATED(5.4, "FOpenGLFrontend is deprecated and derived implementations no longer possible; OpenGL shader compilation should only be triggered from within the shader format code.")
 	void PrecompileShader(FShaderCompilerOutput& ShaderOutput, const FShaderCompilerInput& ShaderInput, const ANSICHAR* ShaderSource, GLSLVersion Version, EHlslShaderFrequency Frequency);
 
+	UE_DEPRECATED(5.4, "FOpenGLFrontend is deprecated and derived implementations no longer possible; OpenGL shader compilation should only be triggered from within the shader format code.")
 	bool PlatformSupportsOfflineCompilation(const GLSLVersion ShaderVersion) const;
 
 	// fills device capabilities in 'offline', mostly hardcoded values
+	UE_DEPRECATED(5.4, "FOpenGLFrontend is deprecated and derived implementations no longer possible; OpenGL shader compilation should only be triggered from within the shader format code.")
 	void FillDeviceCapsOfflineCompilation(struct FDeviceCapabilities &Caps, const GLSLVersion ShaderVersion) const;
 	// final source code processing, based on device capabilities, before actual (offline) compilation; it mostly mirrors the behaviour of OpenGLShaders.cpp/GLSLToDeviceCompatibleGLSL()
+	UE_DEPRECATED(5.4, "FOpenGLFrontend is deprecated and derived implementations no longer possible; OpenGL shader compilation should only be triggered from within the shader format code.")
 	TSharedPtr<ANSICHAR> PrepareCodeForOfflineCompilation(GLSLVersion ShaderVersion, EShaderFrequency Frequency, const ANSICHAR* InShaderSource) const;
 
+	UE_DEPRECATED(5.4, "FOpenGLFrontend is deprecated and derived implementations no longer possible; OpenGL shader compilation should only be triggered from within the shader format code.")
 	void CompileOffline(const FShaderCompilerInput& ShaderInput, FShaderCompilerOutput &Output, const GLSLVersion ShaderVersion, const ANSICHAR *InShaderSource);
 
 	// based on ShaderVersion decides what platform specific compiler to use
+	UE_DEPRECATED(5.4, "FOpenGLFrontend is deprecated and derived implementations no longer possible; OpenGL shader compilation should only be triggered from within the shader format code.")
 	void PlatformCompileOffline(const FShaderCompilerInput &Input, FShaderCompilerOutput& ShaderOutput, const ANSICHAR* ShaderSource, const GLSLVersion ShaderVersion);
 };
