[shaders] implement independent preprocessing IShaderFormat API for OpenGL shader format

diff --git a/Engine/Source/Developer/ShaderFormatOpenGL/Private/OpenGLShaderCompiler.cpp b/Engine/Source/Developer/ShaderFormatOpenGL/Private/OpenGLShaderCompiler.cpp
--- a/Engine/Source/Developer/ShaderFormatOpenGL/Private/OpenGLShaderCompiler.cpp
+++ b/Engine/Source/Developer/ShaderFormatOpenGL/Private/OpenGLShaderCompiler.cpp
@@ -10,7 +10,7 @@
 #include "Serialization/MemoryWriter.h"
 #include "ShaderCompilerDefinitions.h"
 #include "ShaderFormatOpenGL.h"
-#include "ShaderParameterParser.h"
+#include "ShaderPreprocessTypes.h"
 #include "SpirvReflectCommon.h"
 #include <algorithm>
 #include <regex>
@@ -18,7 +18,6 @@
 #if PLATFORM_WINDOWS
 #include "Windows/AllowWindowsPlatformTypes.h"
 	#include <objbase.h>
-	#include <assert.h>
 	#include <stdio.h>
 #include "Windows/HideWindowsPlatformTypes.h"
 #endif
@@ -29,7 +28,6 @@
 #if PLATFORM_WINDOWS
 #include "Windows/AllowWindowsPlatformTypes.h"
 	#include <GL/glcorearb.h>
-	#include <GL/glext.h>
 	#include <GL/wglext.h>
 #include "Windows/HideWindowsPlatformTypes.h"
 #elif PLATFORM_LINUX
@@ -59,7 +57,7 @@
 	#define GL_TESS_CONTROL_SHADER 0x8E88
 	#endif
 #endif
-	#include "OpenGLUtil.h"
+#include "OpenGLUtil.h"
 #include "OpenGLShaderResources.h"
 
 #ifndef DXC_SUPPORTED
@@ -1058,35 +1056,47 @@ void FOpenGLFrontend::PrecompileShader(FShaderCompilerOutput& ShaderOutput, cons
 	PrecompileShaderInternal(ShaderOutput, ShaderInput, ShaderSource, Version, Frequency);
 }
 
+static EHlslCompileTarget GetCompileTarget(GLSLVersion Version)
+{
+	switch (Version)
+	{
+		case GLSL_ES3_1_ANDROID:
+		case GLSL_150_ES3_1:
+			return HCT_FeatureLevelES3_1;
+
+		default:
+			checkNoEntry();
+	}
+	return HCT_InvalidTarget;
+}
+
 PRAGMA_DISABLE_DEPRECATION_WARNINGS		// FShaderCompilerDefinitions will be made internal in the future, marked deprecated until then
 
-static void SetupPerVersionCompilationEnvironmentInternal(GLSLVersion Version, FShaderCompilerDefinitions& AdditionalDefines, EHlslCompileTarget& HlslCompilerTarget)
+static void SetupDefines(GLSLVersion Version, FShaderCompilerDefinitions& AdditionalDefines)
 {
 	switch (Version)
 	{
 		case GLSL_ES3_1_ANDROID:
 			AdditionalDefines.SetDefine(TEXT("COMPILER_GLSL_ES3_1"), 1);
 			AdditionalDefines.SetDefine(TEXT("ES3_1_PROFILE"), 1);
-			HlslCompilerTarget = HCT_FeatureLevelES3_1;
 			break;
 
 		case GLSL_150_ES3_1:
 			AdditionalDefines.SetDefine(TEXT("COMPILER_GLSL"), 1);
 			AdditionalDefines.SetDefine(TEXT("ES3_1_PROFILE"), 1);
-			HlslCompilerTarget = HCT_FeatureLevelES3_1;
 			AdditionalDefines.SetDefine(TEXT("row_major"), TEXT(""));
 			break;
 
 		default:
 			check(0);
 	}
-
 	AdditionalDefines.SetDefine(TEXT("OPENGL_PROFILE"), 1);
 }
 
 void FOpenGLFrontend::SetupPerVersionCompilationEnvironment(GLSLVersion Version, FShaderCompilerDefinitions& AdditionalDefines, EHlslCompileTarget& HlslCompilerTarget)
 {
-	SetupPerVersionCompilationEnvironmentInternal(Version, AdditionalDefines, HlslCompilerTarget);
+	SetupDefines(Version, AdditionalDefines);
+	HlslCompilerTarget = GetCompileTarget(Version);
 }
 
 PRAGMA_ENABLE_DEPRECATION_WARNINGS
@@ -2948,7 +2958,6 @@ bool GenerateDeferredMobileShaders(std::string& GlslSource, GLSLCompileParameter
 static bool CompileToGlslWithShaderConductor(
 	const FShaderCompilerInput&	Input,
 	FShaderCompilerOutput&		Output,
-	const FString&				WorkingDirectory,
 	GLSLVersion					Version,
 	const EShaderFrequency		Frequency,
 	const FString&				PreprocessedShader,
@@ -3017,8 +3026,6 @@ static bool CompileToGlslWithShaderConductor(
 
 	PRAGMA_ENABLE_DEPRECATION_WARNINGS
 
-	UE::ShaderCompilerCommon::DumpDebugShaderData(Input, PreprocessedShader);
-
 	uint32_t BlendFlags = GetDecalBlendFlags(Input);
 
 	// Load shader source into compiler context
@@ -3141,24 +3148,27 @@ static bool CompileToGlslWithShaderConductor(
 		int32 Pos = 0;
 #if !PLATFORM_MAC
 		TCHAR TextureExternalName[256];
+		TCHAR NextToken[2];
 #else
 		ANSICHAR TextureExternalName[256];
+		ANSICHAR NextToken[2];
 #endif
 		do
 		{
 			Pos = PreprocessedShader.Find(TEXT("TextureExternal"), ESearchCase::CaseSensitive, ESearchDir::FromStart, Pos + 15);
 			if (Pos != INDEX_NONE)
 			{
+				// handle both the case where identifier for TextureExternal declaration immediately precedes a ; and has whitespace separating the two
 #if PLATFORM_WINDOWS
-				if (swscanf_s(&PreprocessedShader[Pos], TEXT("TextureExternal %ls"), TextureExternalName, 256) == 1)
+				if (swscanf_s(&PreprocessedShader[Pos], TEXT("TextureExternal %ls %ls"), TextureExternalName, 256, NextToken, 4))
 #elif PLATFORM_MAC
-				if (sscanf(TCHAR_TO_ANSI(&PreprocessedShader[Pos]), "TextureExternal %s", TextureExternalName) == 1)
+				if (sscanf(TCHAR_TO_ANSI(&PreprocessedShader[Pos]), "TextureExternal %s %s", TextureExternalName, NextToken))
 #else // PLATFORM_LINUX
-				if (swscanf(TCHAR_TO_WCHAR(&PreprocessedShader[Pos]), L"TextureExternal %ls", TextureExternalName) == 1)
+				if (swscanf(TCHAR_TO_WCHAR(&PreprocessedShader[Pos]), L"TextureExternal %ls %ls", TextureExternalName, NextToken))
 #endif
 				{
 					FString Name = TextureExternalName;
-					if (Name.RemoveFromEnd(TEXT(";")))
+					if (Name.RemoveFromEnd(TEXT(";")) || (NextToken[0] == TEXT(';')))
 					{
 						ExternalTextures.Add(Name + TEXT("Sampler"));
 					}
@@ -3235,41 +3245,25 @@ static bool CompileToGlslWithShaderConductor(
 
 #endif // DXC_SUPPORTED
 
-/**
- * Compile a shader for OpenGL on Windows.
- * @param Input - The input shader code and environment.
- * @param Output - Contains shader compilation results upon return.
- */
-void CompileOpenGLShader(const FShaderCompilerInput& Input, FShaderCompilerOutput& Output, const FString& WorkingDirectory, GLSLVersion Version)
+static bool ShouldUseDXC(FShaderCompilerFlags Flags)
 {
-	PRAGMA_DISABLE_DEPRECATION_WARNINGS		// FShaderCompilerDefinitions will be made internal in the future, marked deprecated until then
-
-	FString PreprocessedShader;
-	FShaderCompilerDefinitions AdditionalDefines;
-	EHlslCompileTarget HlslCompilerTarget = HCT_InvalidTarget;
-	ECompilerFlags PlatformFlowControl = CFLAG_AvoidFlowControl;
-
-	// set up compiler env based on version
-	SetupPerVersionCompilationEnvironmentInternal(Version, AdditionalDefines, HlslCompilerTarget);
-
 #if DXC_SUPPORTED
-	const bool bUseSC = Input.Environment.CompilerFlags.Contains(CFLAG_ForceDXC);
+	return Flags.Contains(CFLAG_ForceDXC);
 #else
-	const bool bUseSC = false;
+	return false;
 #endif
+}
 
-	AdditionalDefines.SetDefine(TEXT("COMPILER_HLSLCC"), bUseSC ? 2 : 1);
-
-	const bool bDumpDebugInfo = (Input.DumpDebugInfoPath != TEXT("") && IFileManager::Get().DirectoryExists(*Input.DumpDebugInfoPath));
+bool PreprocessOpenGLShader(const FShaderCompilerInput& Input, const FShaderCompilerEnvironment& Environment, FShaderPreprocessOutput& PreprocessOutput, GLSLVersion Version)
+{
+	PRAGMA_DISABLE_DEPRECATION_WARNINGS // FShaderCompilerDefinitions will be made internal in the future, marked deprecated until then
+	
+	FShaderCompilerDefinitions AdditionalDefines;
+	SetupDefines(Version, AdditionalDefines);
 
-	if (Input.Environment.CompilerFlags.Contains(CFLAG_AvoidFlowControl) || PlatformFlowControl == CFLAG_AvoidFlowControl)
-	{
-		AdditionalDefines.SetDefine(TEXT("COMPILER_SUPPORTS_ATTRIBUTES"), (uint32)1);
-	}
-	else
-	{
-		AdditionalDefines.SetDefine(TEXT("COMPILER_SUPPORTS_ATTRIBUTES"), (uint32)0);
-	}
+	const bool bUseSC = ShouldUseDXC(Environment.CompilerFlags);
+	AdditionalDefines.SetDefine(TEXT("COMPILER_HLSLCC"), bUseSC ? 2 : 1);
+	AdditionalDefines.SetDefine(TEXT("COMPILER_SUPPORTS_ATTRIBUTES"), (uint32)1);
 
 	if (Input.Environment.FullPrecisionInPS)
 	{
@@ -3277,30 +3271,68 @@ void CompileOpenGLShader(const FShaderCompilerInput& Input, FShaderCompilerOutpu
 	}
 
 	PRAGMA_ENABLE_DEPRECATION_WARNINGS
-
+	
 	if (Input.bSkipPreprocessedCache)
 	{
-		if (!FFileHelper::LoadFileToString(PreprocessedShader, *Input.VirtualSourceFilePath))
+		if (!FFileHelper::LoadFileToString(PreprocessOutput.EditSource(), *Input.VirtualSourceFilePath))
 		{
-			return;
+			return false;
 		}
 
 		// Remove const as we are on debug-only mode
-		CrossCompiler::CreateEnvironmentFromResourceTable(PreprocessedShader, (FShaderCompilerEnvironment&)Input.Environment);
+		CrossCompiler::CreateEnvironmentFromResourceTable(PreprocessOutput.GetSource(), const_cast<FShaderCompilerEnvironment&>(Environment));
 	}
 	else
 	{
-		if (!PreprocessShader(PreprocessedShader, Output, Input, AdditionalDefines))
+		if (!PreprocessShader(PreprocessOutput, Input, Environment, AdditionalDefines))
 		{
 			// The preprocessing stage will add any relevant errors.
-			return;
+			return false;
 		}
 	}
 
-	char* GlslShaderSource = NULL;
-	char* ErrorLog = NULL;
+	if (!PreprocessOutput.ParseAndModify(Input, Environment, nullptr))
+	{
+		// The FShaderParameterParser will add any relevant errors.
+		return false;
+	}
+
+	// This requires removing the HLSLCC_NoPreprocess flag later on!
+	RemoveUniformBuffersFromSource(Environment, PreprocessOutput.EditSource());
+
+	// Process TEXT macro.
+	TransformStringIntoCharacterArray(PreprocessOutput.EditSource());
 
-	const EHlslShaderFrequency FrequencyTable[] =
+	// Run the experimental shader minifier
+#if UE_OPENGL_SHADER_COMPILER_ALLOW_DEAD_CODE_REMOVAL
+	if (Input.Environment.CompilerFlags.Contains(CFLAG_RemoveDeadCode))
+	{
+		UE::ShaderCompilerCommon::RemoveDeadCode(PreprocessOutput.EditSource(), Input.EntryPointName, PreprocessOutput.EditErrors());
+	}
+#endif // UE_OPENGL_SHADER_COMPILER_ALLOW_DEAD_CODE_REMOVAL
+	return true;
+}
+
+/**
+ * Compile a shader for OpenGL on Windows.
+ * @param Input - The input shader code and environment.
+ * @param PreprocessOutput - The output of the shader preprocessing phase (see above)
+ * @param Output - Contains shader compilation results upon return.
+ * @param WorkingDirectory - (unused, part of IShaderFormat API)
+ * @param Version - Target GLSL version for this compilation
+ */
+void CompileOpenGLShader(const FShaderCompilerInput& Input, const FShaderPreprocessOutput& PreprocessOutput, FShaderCompilerOutput& Output, const FString& WorkingDirectory, GLSLVersion Version)
+{
+	const FString& PreprocessedShader = PreprocessOutput.GetSource();
+	const EHlslCompileTarget HlslCompilerTarget = GetCompileTarget(Version);
+
+	const bool bUseSC = ShouldUseDXC(Input.Environment.CompilerFlags);
+	const bool bDumpDebugInfo = Input.DumpDebugInfoEnabled();
+	
+	char* GlslShaderSource = nullptr;
+	char* ErrorLog = nullptr;
+
+	constexpr EHlslShaderFrequency FrequencyTable[] =
 	{
 		HSF_VertexShader,
 		HSF_InvalidFrequency,
@@ -3311,8 +3343,7 @@ void CompileOpenGLShader(const FShaderCompilerInput& Input, FShaderCompilerOutpu
 	};
 	static_assert(SF_NumStandardFrequencies == UE_ARRAY_COUNT(FrequencyTable), "NumFrequencies changed. Please update tables.");
 
-	const EShaderFrequency Frequency = (EShaderFrequency)Input.Target.Frequency;
-
+	const EShaderFrequency Frequency = static_cast<EShaderFrequency>(Input.Target.Frequency);
 	const EHlslShaderFrequency HlslFrequency = FrequencyTable[Frequency];
 	if (HlslFrequency == HSF_InvalidFrequency)
 	{
@@ -3325,27 +3356,6 @@ void CompileOpenGLShader(const FShaderCompilerInput& Input, FShaderCompilerOutpu
 		return;
 	}
 
-	FShaderParameterParser ShaderParameterParser(Input.Environment.CompilerFlags);
-	if (!ShaderParameterParser.ParseAndModify(Input, Output.Errors, PreprocessedShader))
-	{
-		// The FShaderParameterParser will add any relevant errors.
-		return;
-	}
-
-	// This requires removing the HLSLCC_NoPreprocess flag later on!
-	RemoveUniformBuffersFromSource(Input.Environment, PreprocessedShader);
-
-	// Process TEXT macro.
-	TransformStringIntoCharacterArray(PreprocessedShader);
-
-	// Run the experimental shader minifier
-#if UE_OPENGL_SHADER_COMPILER_ALLOW_DEAD_CODE_REMOVAL
-	if (Input.Environment.CompilerFlags.Contains(CFLAG_RemoveDeadCode))
-	{
-		UE::ShaderCompilerCommon::RemoveDeadCode(PreprocessedShader, Input.EntryPointName, Output.Errors);
-	}
-#endif // UE_OPENGL_SHADER_COMPILER_ALLOW_DEAD_CODE_REMOVAL
-
 	uint32 CCFlags = CalculateCrossCompilerFlagsInternal(Version, Input.Environment.FullPrecisionInPS, Input.Environment.CompilerFlags);
 
 	// Required as we added the RemoveUniformBuffersFromSource() function (the cross-compiler won't be able to interpret comments w/o a preprocessor)
@@ -3356,17 +3366,15 @@ void CompileOpenGLShader(const FShaderCompilerInput& Input, FShaderCompilerOutpu
 #if DXC_SUPPORTED
 	if (bUseSC)
 	{
-		bCompilationSucceeded = CompileToGlslWithShaderConductor(Input, Output, WorkingDirectory, Version, Frequency, PreprocessedShader, GlslShaderSource);
+		bCompilationSucceeded = CompileToGlslWithShaderConductor(Input, Output, Version, Frequency, PreprocessedShader, GlslShaderSource);
 	}
 	else
 #endif // DXC_SUPPORTED
 	{
-		UE::ShaderCompilerCommon::DumpDebugShaderData(Input, PreprocessedShader);
-
 		CCFlags |= HLSLCC_NoValidation;
 		FGlslCodeBackend* BackEnd = CreateBackendInternal(Version, CCFlags, HlslCompilerTarget);
 
-		bool bDefaultPrecisionIsHalf = (CCFlags & HLSLCC_UseFullPrecisionInPS) == 0;
+		const bool bDefaultPrecisionIsHalf = (CCFlags & HLSLCC_UseFullPrecisionInPS) == 0;
 		FGlslLanguageSpec* LanguageSpec = CreateLanguageSpecInternal(Version, bDefaultPrecisionIsHalf);
 
 		{
@@ -3412,8 +3420,8 @@ void CompileOpenGLShader(const FShaderCompilerInput& Input, FShaderCompilerOutpu
 		if (bDumpDebugInfo && GlslShaderSource != nullptr)
 		{
 			const TCHAR* ShaderFileExt = CrossCompiler::FShaderConductorContext::GetShaderFileExt(CrossCompiler::EShaderConductorLanguage::Essl, Frequency);
-			FString DumpedGlslFile = FString::Printf(TEXT("%s/Output.%s"), *Input.DumpDebugInfoPath, ShaderFileExt);
-			if (TUniquePtr<FArchive> FileWriter = TUniquePtr<FArchive>(IFileManager::Get().CreateFileWriter(*DumpedGlslFile)))
+			const FString DumpedGlslFile = FString::Printf(TEXT("%s/Output.%s"), *Input.DumpDebugInfoPath, ShaderFileExt);
+			if (const TUniquePtr<FArchive> FileWriter = TUniquePtr<FArchive>(IFileManager::Get().CreateFileWriter(*DumpedGlslFile)))
 			{
 				FileWriter->Serialize(GlslShaderSource, FCStringAnsi::Strlen(GlslShaderSource));
 				FileWriter->Close();
@@ -3424,7 +3432,7 @@ void CompileOpenGLShader(const FShaderCompilerInput& Input, FShaderCompilerOutpu
 	{
 		const bool bUseAbsolutePaths = bDirectCompile;
 
-		FString Tmp = ANSI_TO_TCHAR(ErrorLog);
+		const FString Tmp = ANSI_TO_TCHAR(ErrorLog);
 		TArray<FString> ErrorLines;
 		Tmp.ParseIntoArray(ErrorLines, TEXT("\n"), true);
 
@@ -3448,10 +3456,11 @@ void CompileOpenGLShader(const FShaderCompilerInput& Input, FShaderCompilerOutpu
 	//ShaderParameterParser.ValidateShaderParameterTypes(Input, Output);
 }
 
-
 void FOpenGLFrontend::CompileShader(const FShaderCompilerInput& Input, FShaderCompilerOutput& Output, const FString& WorkingDirectory, GLSLVersion Version)
 {
-	CompileOpenGLShader(Input, Output, WorkingDirectory, Version);
+	FShaderPreprocessOutput PreprocessOutput;
+	PreprocessOpenGLShader(Input, Input.Environment, PreprocessOutput, Version);
+	CompileOpenGLShader(Input, PreprocessOutput, Output, WorkingDirectory, Version);
 }
 
 static void FillDeviceCapsOfflineCompilationInternal(struct FDeviceCapabilities& Capabilities, const GLSLVersion ShaderVersion)
diff --git a/Engine/Source/Developer/ShaderFormatOpenGL/Private/ShaderFormatOpenGL.cpp b/Engine/Source/Developer/ShaderFormatOpenGL/Private/ShaderFormatOpenGL.cpp
--- a/Engine/Source/Developer/ShaderFormatOpenGL/Private/ShaderFormatOpenGL.cpp
+++ b/Engine/Source/Developer/ShaderFormatOpenGL/Private/ShaderFormatOpenGL.cpp
@@ -2,6 +2,8 @@
 //
 
 #include "ShaderFormatOpenGL.h"
+
+#include "ShaderCompilerCommon.h"
 #include "HAL/FileManager.h"
 #include "Modules/ModuleManager.h"
 #include "Interfaces/IShaderFormat.h"
@@ -12,9 +14,19 @@
 static FName NAME_GLSL_150_ES3_1(TEXT("GLSL_150_ES31"));
 static FName NAME_GLSL_ES3_1_ANDROID(TEXT("GLSL_ES3_1_ANDROID"));
 
-extern void CompileOpenGLShader(const FShaderCompilerInput& Input, FShaderCompilerOutput& Output, const FString& WorkingDirectory, GLSLVersion Version);
-
-class FShaderFormatGLSL : public IShaderFormat
+extern bool PreprocessOpenGLShader(
+	const FShaderCompilerInput& Input,
+	const FShaderCompilerEnvironment& Environment,
+	FShaderPreprocessOutput& Output,
+	GLSLVersion Version);
+extern void CompileOpenGLShader(
+	const FShaderCompilerInput& Input,
+	const FShaderPreprocessOutput& PreprocessOutput,
+	FShaderCompilerOutput& Output,
+	const FString& WorkingDirectory,
+	GLSLVersion Version);
+
+class FShaderFormatGLSL : public UE::ShaderCompilerCommon::FBaseShaderFormat 
 {
 	enum
 	{
@@ -22,7 +34,7 @@ class FShaderFormatGLSL : public IShaderFormat
 		UE_SHADER_GLSL_VER = 107,
 	};
 
-	void CheckFormat(FName Format) const
+	static void CheckFormat(FName Format)
 	{
 		check(Format == NAME_GLSL_150_ES3_1 || Format == NAME_GLSL_ES3_1_ANDROID);
 	}
@@ -72,16 +84,24 @@ public:
 		}
 	}
 
-	virtual void CompileShader(FName Format, const struct FShaderCompilerInput& Input, struct FShaderCompilerOutput& Output,const FString& WorkingDirectory) const override
+	virtual bool SupportsIndependentPreprocessing() const override
 	{
-		CheckFormat(Format);
+		return true;
+	}
 
-		GLSLVersion Version = TranslateFormatNameToEnum(Format);
+	virtual bool PreprocessShader(const FShaderCompilerInput& Input, const FShaderCompilerEnvironment& Environment, FShaderPreprocessOutput& PreprocessOutput) const override
+	{
+		CheckFormat(Input.ShaderFormat);
+		return PreprocessOpenGLShader(Input, Environment, PreprocessOutput, TranslateFormatNameToEnum(Input.ShaderFormat));
+	}
 
-		CompileOpenGLShader(Input, Output, WorkingDirectory, Version);
+	virtual void CompilePreprocessedShader(const FShaderCompilerInput& Input, const FShaderPreprocessOutput& PreprocessOutput, FShaderCompilerOutput& Output, const FString& WorkingDirectory) const override
+	{
+		CheckFormat(Input.ShaderFormat);
+		CompileOpenGLShader(Input, PreprocessOutput, Output, WorkingDirectory, TranslateFormatNameToEnum(Input.ShaderFormat));		
 	}
 
-	virtual const TCHAR* GetPlatformIncludeDirectory() const
+	virtual const TCHAR* GetPlatformIncludeDirectory() const override
 	{
 		return TEXT("GL");
 	}
@@ -91,17 +111,17 @@ public:
  * Module for OpenGL shaders
  */
 
-static IShaderFormat* Singleton = NULL;
+static IShaderFormat* Singleton = nullptr;
 
 class FShaderFormatOpenGLModule : public IShaderFormatModule
 {
 public:
-	virtual ~FShaderFormatOpenGLModule()
+	virtual ~FShaderFormatOpenGLModule() override
 	{
 		delete Singleton;
-		Singleton = NULL;
+		Singleton = nullptr;
 	}
-	virtual IShaderFormat* GetShaderFormat()
+	virtual IShaderFormat* GetShaderFormat() override
 	{
 		if (!Singleton)
 		{
