Changing DumpDebugInfoPath conditions to use a helper function.

diff --git a/Engine/Source/Developer/ShaderFormatOpenGL/Private/OpenGLShaderCompiler.cpp b/Engine/Source/Developer/ShaderFormatOpenGL/Private/OpenGLShaderCompiler.cpp
--- a/Engine/Source/Developer/ShaderFormatOpenGL/Private/OpenGLShaderCompiler.cpp
+++ b/Engine/Source/Developer/ShaderFormatOpenGL/Private/OpenGLShaderCompiler.cpp
@@ -2965,7 +2965,7 @@ static bool CompileToGlslWithShaderConductor(
 {
 	CrossCompiler::FShaderConductorContext CompilerContext;
 
-	const bool bDumpDebugInfo = (Input.DumpDebugInfoPath != TEXT("") && IFileManager::Get().DirectoryExists(*Input.DumpDebugInfoPath));
+	const bool bDumpDebugInfo = Input.DumpDebugInfoEnabled();
 	const bool bRewriteHlslSource = true;
 
 	// Initialize compilation options for ShaderConductor
