Fix crash in OpenGL shader compiler on macOS

diff --git a/Engine/Source/Developer/ShaderFormatOpenGL/Private/OpenGLShaderCompiler.cpp b/Engine/Source/Developer/ShaderFormatOpenGL/Private/OpenGLShaderCompiler.cpp
--- a/Engine/Source/Developer/ShaderFormatOpenGL/Private/OpenGLShaderCompiler.cpp
+++ b/Engine/Source/Developer/ShaderFormatOpenGL/Private/OpenGLShaderCompiler.cpp
@@ -3171,11 +3171,11 @@ static bool CompileToGlslWithShaderConductor(
 		TSet<FString> ExternalTextures;
 		int32 Pos = 0;
 #if !PLATFORM_MAC
-		TCHAR TextureExternalName[256];
-		TCHAR NextToken[2];
+		TCHAR TextureExternalName[257] = {};
+		TCHAR NextToken[2] = {};
 #else
-		ANSICHAR TextureExternalName[256];
-		ANSICHAR NextToken[2];
+		ANSICHAR TextureExternalName[257] = {};
+		ANSICHAR NextToken[2] = {};
 #endif
 		do
 		{
@@ -3186,9 +3186,9 @@ static bool CompileToGlslWithShaderConductor(
 #if PLATFORM_WINDOWS
 				if (swscanf_s(&PreprocessedShader[Pos], TEXT("TextureExternal %ls %ls"), TextureExternalName, 256, NextToken, 4))
 #elif PLATFORM_MAC
-				if (sscanf(TCHAR_TO_ANSI(&PreprocessedShader[Pos]), "TextureExternal %s %s", TextureExternalName, NextToken))
+				if (sscanf(TCHAR_TO_ANSI(&PreprocessedShader[Pos]), "TextureExternal %256s %1s", TextureExternalName, NextToken))
 #else // PLATFORM_LINUX
-				if (swscanf(TCHAR_TO_WCHAR(&PreprocessedShader[Pos]), L"TextureExternal %ls %ls", TextureExternalName, NextToken))
+				if (swscanf(TCHAR_TO_WCHAR(&PreprocessedShader[Pos]), L"TextureExternal %256ls %1ls", TextureExternalName, NextToken))
 #endif
 				{
 					FString Name = TextureExternalName;
