Fix for incorrect regex expression when renaming uniform buffers in OpenGL

diff --git a/Engine/Shaders/Public/Platform/GL/GLCommon.ush b/Engine/Shaders/Public/Platform/GL/GLCommon.ush
--- a/Engine/Shaders/Public/Platform/GL/GLCommon.ush
+++ b/Engine/Shaders/Public/Platform/GL/GLCommon.ush
@@ -7,6 +7,6 @@
 #pragma once
 
 // Update this GUID to invalidate shader recompilation for only GL shaders
-#pragma message("UESHADERMETADATA_VERSION 1535AA07-356D-4F56-88BD-B38BDD770D36")
+#pragma message("UESHADERMETADATA_VERSION 3FB1EB49-0442-40E1-9F1B-0371F7ABB4BB")
 
 
diff --git a/Engine/Source/Developer/ShaderFormatOpenGL/Private/OpenGLShaderCompiler.cpp b/Engine/Source/Developer/ShaderFormatOpenGL/Private/OpenGLShaderCompiler.cpp
--- a/Engine/Source/Developer/ShaderFormatOpenGL/Private/OpenGLShaderCompiler.cpp
+++ b/Engine/Source/Developer/ShaderFormatOpenGL/Private/OpenGLShaderCompiler.cpp
@@ -2509,25 +2509,10 @@ bool GenerateGlslShader(std::string& OutString, GLSLCompileParameters& GLSLCompi
 			std::string NewUAVName = FrequencyPrefix + std::string("i") + std::to_string(UAVIndex);
 
 			// Find instances of UAVs
-			std::string RegexExpression = "(?:^|\\W|\\S)" + UAV + "(?:$|\\W)";
+			std::string RegexExpression = "\\b" + UAV + "\\b";
+			std::regex RegexExpressionPattern(RegexExpression);
 
-			std::cmatch RegexMatch;
-			std::vector<std::string> RegexMatches;
-
-			const char* TempString = OutString.c_str();
-			while (std::regex_search(TempString, RegexMatch, std::regex(RegexExpression)))
-			{
-				RegexMatches.push_back(RegexMatch.str());
-				TempString = RegexMatch.suffix().first;
-			}
-
-			for (const auto& Match : RegexMatches)
-			{
-				size_t MatchOffset = Match.find_first_of(UAV);
-				size_t MatchPos = OutString.find(Match);
-
-				OutString.replace(MatchPos + MatchOffset, UAV.size(), NewUAVName);
-			}
+			OutString = std::regex_replace(OutString, RegexExpressionPattern, NewUAVName);
 
 			UAVIndex++;
 		}
