[Shaders] Remove Hlslcc form GL backend; Mobile renderer fully switched to DXC.

#rnx

diff --git a/Engine/Source/Developer/ShaderFormatOpenGL/Private/GlslBackend.cpp b/Engine/Source/Developer/ShaderFormatOpenGL/Private/GlslBackend.cpp
deleted file mode 100644
diff --git a/Engine/Source/Developer/ShaderFormatOpenGL/Private/OpenGLShaderCompiler.cpp b/Engine/Source/Developer/ShaderFormatOpenGL/Private/OpenGLShaderCompiler.cpp
--- a/Engine/Source/Developer/ShaderFormatOpenGL/Private/OpenGLShaderCompiler.cpp
+++ b/Engine/Source/Developer/ShaderFormatOpenGL/Private/OpenGLShaderCompiler.cpp
@@ -25,7 +25,6 @@
 #include "ShaderCore.h"
 #include "ShaderPreprocessor.h"
 #include "ShaderCompilerCommon.h"
-#include "GlslBackend.h"
 #if PLATFORM_WINDOWS
 #include "Windows/AllowWindowsPlatformTypes.h"
 	#include <GL/glcorearb.h>
@@ -61,16 +60,10 @@
 #include "OpenGLUtil.h"
 #include "OpenGLShaderResources.h"
 
-#ifndef DXC_SUPPORTED
-	#define DXC_SUPPORTED (PLATFORM_MAC || PLATFORM_WINDOWS || PLATFORM_LINUX) 
-#endif
-
-#if DXC_SUPPORTED
 THIRD_PARTY_INCLUDES_START
 #include "spirv_reflect.h"
 #include <map>
 THIRD_PARTY_INCLUDES_END
-#endif // DXC_SUPPORTED
 
 DEFINE_LOG_CATEGORY_STATIC(LogOpenGLShaderCompiler, Log, All);
 
@@ -912,16 +905,6 @@ void BuildShaderOutputInternal(
 	}
 }
 
-void FOpenGLFrontend::BuildShaderOutput(
-	FShaderCompilerOutput& ShaderOutput,
-	const FShaderCompilerInput& ShaderInput,
-	const ANSICHAR* InShaderSource,
-	int32 SourceLen,
-	GLSLVersion Version)
-{
-	BuildShaderOutputInternal(ShaderOutput, ShaderInput, InShaderSource, SourceLen, Version);
-}
-
 static void ConvertOpenGLVersionFromGLSLVersionInternal(GLSLVersion InVersion, int& OutMajorVersion, int& OutMinorVersion)
 {
 	switch(InVersion)
@@ -943,11 +926,6 @@ static void ConvertOpenGLVersionFromGLSLVersionInternal(GLSLVersion InVersion, i
 	}
 }
 
-void FOpenGLFrontend::ConvertOpenGLVersionFromGLSLVersion(GLSLVersion InVersion, int& OutMajorVersion, int& OutMinorVersion)
-{
-	return ConvertOpenGLVersionFromGLSLVersionInternal(InVersion, OutMajorVersion, OutMinorVersion);
-}
-
 /**
  * Precompile a GLSL shader.
  * @param ShaderOutput - The precompiled shader.
@@ -1046,11 +1024,6 @@ static void PrecompileShaderInternal(FShaderCompilerOutput& ShaderOutput, const
 	PlatformReleaseOpenGL(ContextPtr, PrevContextPtr);
 }
 
-void FOpenGLFrontend::PrecompileShader(FShaderCompilerOutput& ShaderOutput, const FShaderCompilerInput& ShaderInput, const ANSICHAR* ShaderSource, GLSLVersion Version, EHlslShaderFrequency Frequency)
-{
-	PrecompileShaderInternal(ShaderOutput, ShaderInput, ShaderSource, Version, Frequency);
-}
-
 static EHlslCompileTarget GetCompileTarget(GLSLVersion Version)
 {
 	switch (Version)
@@ -1065,21 +1038,6 @@ static EHlslCompileTarget GetCompileTarget(GLSLVersion Version)
 	return HCT_InvalidTarget;
 }
 
-PRAGMA_DISABLE_DEPRECATION_WARNINGS		// FShaderCompilerDefinitions will be made internal in the future, marked deprecated until then
-
-void FOpenGLFrontend::SetupPerVersionCompilationEnvironment(GLSLVersion Version, FShaderCompilerDefinitions& AdditionalDefines, EHlslCompileTarget& HlslCompilerTarget)
-{
-	// this function is deprecated; SetupDefines is removed and all per-version defines are now set in ModifyShaderCompilerInput on the IShaderFormat implementation.
-	HlslCompilerTarget = GetCompileTarget(Version);
-}
-
-PRAGMA_ENABLE_DEPRECATION_WARNINGS
-
-uint32 FOpenGLFrontend::GetMaxSamplers(GLSLVersion Version)
-{
-	return 16;
-}
-
 static uint32 CalculateCrossCompilerFlagsInternal(GLSLVersion Version, const bool bFullPrecisionInPS, const FShaderCompilerFlags& CompilerFlags)
 {
 	uint32  CCFlags = HLSLCC_NoPreprocess | HLSLCC_PackUniforms | HLSLCC_DX11ClipSpace | HLSLCC_RetainSizes;
@@ -1105,43 +1063,6 @@ static uint32 CalculateCrossCompilerFlagsInternal(GLSLVersion Version, const boo
 	return CCFlags;
 }
 
-uint32 FOpenGLFrontend::CalculateCrossCompilerFlags(GLSLVersion Version, const bool bFullPrecisionInPS, const FShaderCompilerFlags& CompilerFlags)
-{
-	return CalculateCrossCompilerFlagsInternal(Version, bFullPrecisionInPS, CompilerFlags);
-}
-
-FGlslCodeBackend* CreateBackendInternal(GLSLVersion Version, uint32 CCFlags, EHlslCompileTarget HlslCompilerTarget)
-{
-	return new FGlslCodeBackend(CCFlags, HlslCompilerTarget);
-}
-
-
-FGlslCodeBackend* FOpenGLFrontend::CreateBackend(GLSLVersion Version, uint32 CCFlags, EHlslCompileTarget HlslCompilerTarget)
-{
-	return CreateBackendInternal(Version, CCFlags, HlslCompilerTarget);
-}
-
-class FGlsl430LanguageSpec : public FGlslLanguageSpec
-{
-public:
-	FGlsl430LanguageSpec(bool bInDefaultPrecisionIsHalf)
-		: FGlslLanguageSpec(bInDefaultPrecisionIsHalf)
-	{}
-	virtual bool EmulateStructuredWithTypedBuffers() const override { return false; }
-};
-
-FGlslLanguageSpec* CreateLanguageSpecInternal(GLSLVersion Version, bool bDefaultPrecisionIsHalf)
-{
-	return new FGlslLanguageSpec(bDefaultPrecisionIsHalf);
-}
-
-FGlslLanguageSpec* FOpenGLFrontend::CreateLanguageSpec(GLSLVersion Version, bool bDefaultPrecisionIsHalf)
-{
-	return CreateLanguageSpecInternal(Version, bDefaultPrecisionIsHalf);
-}
-
-#if DXC_SUPPORTED
-
 static const ANSICHAR* GetFrequencyPrefix(EShaderFrequency Frequency)
 {
 	switch (Frequency)
@@ -3231,17 +3152,6 @@ static bool CompileToGlslWithShaderConductor(
 	return !bCompilationFailed;
 }
 
-#endif // DXC_SUPPORTED
-
-bool ShouldUseDXC(FShaderCompilerFlags Flags)
-{
-#if DXC_SUPPORTED
-	return Flags.Contains(CFLAG_ForceDXC);
-#else
-	return false;
-#endif
-}
-
 /**
  * Compile a shader for OpenGL on Windows.
  * @param Input - The input shader code and environment.
@@ -3270,7 +3180,6 @@ void CompileOpenGLShader(const FShaderCompilerInput& Input, const FShaderPreproc
 
 	const EHlslCompileTarget HlslCompilerTarget = GetCompileTarget(Version);
 
-	const bool bUseSC = ShouldUseDXC(Input.Environment.CompilerFlags);
 	const bool bDumpDebugInfo = Input.DumpDebugInfoEnabled();
 	
 	char* GlslShaderSource = nullptr;
@@ -3305,45 +3214,7 @@ void CompileOpenGLShader(const FShaderCompilerInput& Input, const FShaderPreproc
 	// Required as we added the RemoveUniformBuffersFromSource() function (the cross-compiler won't be able to interpret comments w/o a preprocessor)
 	CCFlags &= ~HLSLCC_NoPreprocess;
 
-	bool bCompilationSucceeded = false;
-
-#if DXC_SUPPORTED
-	if (bUseSC)
-	{
-		bCompilationSucceeded = CompileToGlslWithShaderConductor(Input, Output, Version, Frequency, PreprocessedSource, GlslShaderSource);
-	}
-	else
-#endif // DXC_SUPPORTED
-	{
-		CCFlags |= HLSLCC_NoValidation;
-		FGlslCodeBackend* BackEnd = CreateBackendInternal(Version, CCFlags, HlslCompilerTarget);
-
-		const bool bDefaultPrecisionIsHalf = (CCFlags & HLSLCC_UseFullPrecisionInPS) == 0;
-		FGlslLanguageSpec* LanguageSpec = CreateLanguageSpecInternal(Version, bDefaultPrecisionIsHalf);
-
-		{
-			FScopeLock HlslCcLock(CrossCompiler::GetCrossCompilerLock());
-			FHlslCrossCompilerContext CrossCompilerContext(CCFlags, HlslFrequency, HlslCompilerTarget);
-			if (CrossCompilerContext.Init(TCHAR_TO_ANSI(*Input.VirtualSourceFilePath), LanguageSpec))
-			{
-				bCompilationSucceeded = CrossCompilerContext.Run(
-					TCHAR_TO_ANSI(*PreprocessedSource),
-					TCHAR_TO_ANSI(*Input.EntryPointName),
-					BackEnd,
-					&GlslShaderSource,
-					&ErrorLog
-				);
-			}
-		}
-		delete BackEnd;
-		delete LanguageSpec;
-
-		if (bDumpDebugInfo && bCompilationSucceeded && GlslShaderSource != nullptr)
-		{
-			const TCHAR* ShaderFileExt = CrossCompiler::FShaderConductorContext::GetShaderFileExt(CrossCompiler::EShaderConductorLanguage::Essl, Frequency);
-			DumpDebugShaderText(Input, GlslShaderSource, FCStringAnsi::Strlen(GlslShaderSource), ShaderFileExt);
-		}
-	}
+	const bool bCompilationSucceeded = CompileToGlslWithShaderConductor(Input, Output, Version, Frequency, PreprocessedSource, GlslShaderSource);
 
 	if (bCompilationSucceeded)
 	{
@@ -3372,20 +3243,6 @@ void CompileOpenGLShader(const FShaderCompilerInput& Input, const FShaderPreproc
 			}
 		}
 	}
-	else if (!bUseSC)
-	{
-		const bool bUseAbsolutePaths = EnumHasAnyFlags(Input.DebugInfoFlags, EShaderDebugInfoFlags::CompileFromDebugUSF);
-
-		const FString Tmp = ANSI_TO_TCHAR(ErrorLog);
-		TArray<FString> ErrorLines;
-		Tmp.ParseIntoArray(ErrorLines, TEXT("\n"), true);
-
-		for (int32 LineIndex = 0; LineIndex < ErrorLines.Num(); ++LineIndex)
-		{
-			const FString& Line = ErrorLines[LineIndex];
-			CrossCompiler::ParseHlslccError(Output.Errors, Line, bUseAbsolutePaths);
-		}
-	}
 
 	if (GlslShaderSource)
 	{
@@ -3400,13 +3257,6 @@ void CompileOpenGLShader(const FShaderCompilerInput& Input, const FShaderPreproc
 	//ShaderParameterParser.ValidateShaderParameterTypes(Input, Output);
 }
 
-void FOpenGLFrontend::CompileShader(const FShaderCompilerInput& Input, FShaderCompilerOutput& Output, const FString& WorkingDirectory, GLSLVersion Version)
-{
-	FShaderPreprocessOutput PreprocessOutput;
-	PreprocessShader(PreprocessOutput, Input, Input.Environment);
-	CompileOpenGLShader(Input, PreprocessOutput, Output, WorkingDirectory, Version);
-}
-
 static void FillDeviceCapsOfflineCompilationInternal(struct FDeviceCapabilities& Capabilities, const GLSLVersion ShaderVersion)
 {
 	FMemory::Memzero(Capabilities);
@@ -3421,11 +3271,6 @@ static void FillDeviceCapsOfflineCompilationInternal(struct FDeviceCapabilities&
 	}
 }
 
-void FOpenGLFrontend::FillDeviceCapsOfflineCompilation(struct FDeviceCapabilities& Capabilities, const GLSLVersion ShaderVersion) const
-{
-	FillDeviceCapsOfflineCompilationInternal(Capabilities, ShaderVersion);
-}
-
 static bool MoveHashLines(FString& Destination, FString &Source)
 {
 	int32 Index = 0;
@@ -3541,21 +3386,11 @@ static TSharedPtr<ANSICHAR> PrepareCodeForOfflineCompilationInternal(const GLSLV
 	return RetShaderSource;
 }
 
-TSharedPtr<ANSICHAR> FOpenGLFrontend::PrepareCodeForOfflineCompilation(const GLSLVersion ShaderVersion, EShaderFrequency Frequency, const ANSICHAR* InShaderSource) const
-{
-	return PrepareCodeForOfflineCompilationInternal(ShaderVersion, Frequency, InShaderSource);
-}
-
 static bool PlatformSupportsOfflineCompilationInternal(const GLSLVersion ShaderVersion)
 {
 	return ShaderVersion == GLSL_ES3_1_ANDROID;
 }
 
-bool FOpenGLFrontend::PlatformSupportsOfflineCompilation(const GLSLVersion ShaderVersion) const
-{
-	return PlatformSupportsOfflineCompilationInternal(ShaderVersion);
-}
-
 static void CompileOfflineInternal(const FShaderCompilerInput& Input, FShaderCompilerOutput& Output, const GLSLVersion ShaderVersion, const ANSICHAR* InShaderSource)
 {
 	const bool bSupportsOfflineCompilation = PlatformSupportsOfflineCompilationInternal(ShaderVersion);
@@ -3570,11 +3405,6 @@ static void CompileOfflineInternal(const FShaderCompilerInput& Input, FShaderCom
 	PlatformCompileOfflineInternal(Input, Output, ShaderSource.Get(), ShaderVersion);
 }
 
-void FOpenGLFrontend::CompileOffline(const FShaderCompilerInput& Input, FShaderCompilerOutput& Output, const GLSLVersion ShaderVersion, const ANSICHAR* InShaderSource)
-{
-	CompileOfflineInternal(Input, Output, ShaderVersion, InShaderSource);
-}
-
 static void PlatformCompileOfflineInternal(const FShaderCompilerInput& Input, FShaderCompilerOutput& ShaderOutput, const ANSICHAR* ShaderSource, const GLSLVersion ShaderVersion)
 {
 	if (ShaderVersion == GLSL_ES3_1_ANDROID)
@@ -3582,8 +3412,3 @@ static void PlatformCompileOfflineInternal(const FShaderCompilerInput& Input, FS
 		CompileOfflineMali(Input, ShaderOutput, ShaderSource, FPlatformString::Strlen(ShaderSource), false);
 	}
 }
-
-void FOpenGLFrontend::PlatformCompileOffline(const FShaderCompilerInput& Input, FShaderCompilerOutput& ShaderOutput, const ANSICHAR* ShaderSource, const GLSLVersion ShaderVersion)
-{
-	PlatformCompileOfflineInternal(Input, ShaderOutput, ShaderSource, ShaderVersion);
-}
diff --git a/Engine/Source/Developer/ShaderFormatOpenGL/Private/ShaderFormatOpenGL.cpp b/Engine/Source/Developer/ShaderFormatOpenGL/Private/ShaderFormatOpenGL.cpp
--- a/Engine/Source/Developer/ShaderFormatOpenGL/Private/ShaderFormatOpenGL.cpp
+++ b/Engine/Source/Developer/ShaderFormatOpenGL/Private/ShaderFormatOpenGL.cpp
@@ -96,8 +96,8 @@ public:
 		}
 		Input.Environment.SetDefine(TEXT("OPENGL_PROFILE"), 1);
 
-		const bool bUseDXC = ShouldUseDXC(Input.Environment.CompilerFlags);
-		Input.Environment.SetDefine(TEXT("COMPILER_HLSLCC"), bUseDXC ? 2 : 1);
+		constexpr int32 HlslccSubstituteValueForDxc = 2;
+		Input.Environment.SetDefine(TEXT("COMPILER_HLSLCC"), HlslccSubstituteValueForDxc);
 		Input.Environment.SetDefine(TEXT("COMPILER_SUPPORTS_ATTRIBUTES"), (uint32)1);
 
 		if (Input.Environment.FullPrecisionInPS || (IsValidRef(Input.SharedEnvironment) && Input.SharedEnvironment->FullPrecisionInPS))
diff --git a/Engine/Source/Developer/ShaderFormatOpenGL/Public/GlslBackend.h b/Engine/Source/Developer/ShaderFormatOpenGL/Public/GlslBackend.h
deleted file mode 100644
diff --git a/Engine/Source/Developer/ShaderFormatOpenGL/Public/ShaderFormatOpenGL.h b/Engine/Source/Developer/ShaderFormatOpenGL/Public/ShaderFormatOpenGL.h
--- a/Engine/Source/Developer/ShaderFormatOpenGL/Public/ShaderFormatOpenGL.h
+++ b/Engine/Source/Developer/ShaderFormatOpenGL/Public/ShaderFormatOpenGL.h
@@ -38,94 +38,3 @@ enum GLSLVersion
 
 	GLSL_MAX
 };
-
-class SHADERFORMATOPENGL_API FOpenGLFrontend
-{
-public:
-	UE_DEPRECATED(5.4, "FOpenGLFrontend is deprecated and derived implementations no longer possible; OpenGL shader compilation should only be triggered from within the shader format code.")
-	virtual ~FOpenGLFrontend()
-	{
-
-	}
-
-	UE_DEPRECATED(5.4, "FOpenGLFrontend is deprecated and derived implementations no longer possible; OpenGL shader compilation should only be triggered from within the shader format code.")
-	void CompileShader(const struct FShaderCompilerInput& Input, struct FShaderCompilerOutput& Output, const class FString& WorkingDirectory, GLSLVersion Version);
-
-protected:
-	// if true, the shader output map will contain true names (i.e. ColorModifier) instead of helper names for runtime binding (i.e. pb_5)
-	UE_DEPRECATED(5.4, "FOpenGLFrontend is deprecated and derived implementations no longer possible; OpenGL shader compilation should only be triggered from within the shader format code.")
-	virtual bool OutputTrueParameterNames()
-	{
-		return false;
-	}
-
-	UE_DEPRECATED(5.4, "FOpenGLFrontend is deprecated and derived implementations no longer possible; OpenGL shader compilation should only be triggered from within the shader format code.")
-	virtual bool IsSM5(GLSLVersion Version)
-	{
-		return false;
-	}
-
-	// what is the max number of samplers the shader platform can use?
-	UE_DEPRECATED(5.4, "FOpenGLFrontend is deprecated and derived implementations no longer possible; OpenGL shader compilation should only be triggered from within the shader format code.")
-	virtual uint32 GetMaxSamplers(GLSLVersion Version);
-
-	UE_DEPRECATED(5.4, "FOpenGLFrontend is deprecated and derived implementations no longer possible; OpenGL shader compilation should only be triggered from within the shader format code.")
-	virtual uint32 CalculateCrossCompilerFlags(GLSLVersion Version, const bool bFullPrecisionInPS, const FShaderCompilerFlags& CompilerFlags);
-
-	// set up compilation information like defines and HlslCompileTarget
-	PRAGMA_DISABLE_DEPRECATION_WARNINGS		// FShaderCompilerDefinitions will be made internal in the future, marked deprecated until then
-	UE_DEPRECATED(5.4, "FOpenGLFrontend is deprecated and derived implementations no longer possible; OpenGL shader compilation should only be triggered from within the shader format code.")
-	virtual void SetupPerVersionCompilationEnvironment(GLSLVersion Version, class FShaderCompilerDefinitions& AdditionalDefines, EHlslCompileTarget& HlslCompilerTarget);
-	PRAGMA_ENABLE_DEPRECATION_WARNINGS
-
-	UE_DEPRECATED(5.4, "FOpenGLFrontend is deprecated and derived implementations no longer possible; OpenGL shader compilation should only be triggered from within the shader format code.")
-	virtual void ConvertOpenGLVersionFromGLSLVersion(GLSLVersion InVersion, int& OutMajorVersion, int& OutMinorVersion);
-
-	// create the compiling backend
-	UE_DEPRECATED(5.4, "FOpenGLFrontend is deprecated and derived implementations no longer possible; OpenGL shader compilation should only be triggered from within the shader format code.")
-	virtual struct FGlslCodeBackend* CreateBackend(GLSLVersion Version, uint32 CCFlags, EHlslCompileTarget HlslCompilerTarget);
-
-	// create the language spec
-	UE_DEPRECATED(5.4, "FOpenGLFrontend is deprecated and derived implementations no longer possible; OpenGL shader compilation should only be triggered from within the shader format code.")
-	virtual class FGlslLanguageSpec* CreateLanguageSpec(GLSLVersion Version, bool bDefaultPrecisionIsHalf);
-
-
-	// Allow a subclass to perform additional work on the cross compiled source code
-	UE_DEPRECATED(5.4, "FOpenGLFrontend is deprecated and derived implementations no longer possible; OpenGL shader compilation should only be triggered from within the shader format code.")
-	virtual bool PostProcessShaderSource(GLSLVersion Version, EShaderFrequency Frequency, const ANSICHAR* ShaderSource,
-		uint32 SourceLen, class FShaderParameterMap& ParameterMap, TMap<FString, FString>& BindingNameMap, TArray<struct FShaderCompilerError>& Errors,
-		const FShaderCompilerInput& ShaderInput)
-	{
-		return true;
-	}
-
-	// allow subclass to write out different output, returning true if it did write everything it needed
-	UE_DEPRECATED(5.4, "FOpenGLFrontend is deprecated and derived implementations no longer possible; OpenGL shader compilation should only be triggered from within the shader format code.")
-	virtual bool OptionalSerializeOutputAndReturnIfSerialized(FArchive& Ar)
-	{
-		return false;
-	}
-
-
-	UE_DEPRECATED(5.4, "FOpenGLFrontend is deprecated and derived implementations no longer possible; OpenGL shader compilation should only be triggered from within the shader format code.")
-	void BuildShaderOutput(FShaderCompilerOutput& ShaderOutput, const FShaderCompilerInput& ShaderInput, const ANSICHAR* InShaderSource, int32 SourceLen, GLSLVersion Version);
-	UE_DEPRECATED(5.4, "FOpenGLFrontend is deprecated and derived implementations no longer possible; OpenGL shader compilation should only be triggered from within the shader format code.")
-	void PrecompileShader(FShaderCompilerOutput& ShaderOutput, const FShaderCompilerInput& ShaderInput, const ANSICHAR* ShaderSource, GLSLVersion Version, EHlslShaderFrequency Frequency);
-
-	UE_DEPRECATED(5.4, "FOpenGLFrontend is deprecated and derived implementations no longer possible; OpenGL shader compilation should only be triggered from within the shader format code.")
-	bool PlatformSupportsOfflineCompilation(const GLSLVersion ShaderVersion) const;
-
-	// fills device capabilities in 'offline', mostly hardcoded values
-	UE_DEPRECATED(5.4, "FOpenGLFrontend is deprecated and derived implementations no longer possible; OpenGL shader compilation should only be triggered from within the shader format code.")
-	void FillDeviceCapsOfflineCompilation(struct FDeviceCapabilities &Caps, const GLSLVersion ShaderVersion) const;
-	// final source code processing, based on device capabilities, before actual (offline) compilation; it mostly mirrors the behaviour of OpenGLShaders.cpp/GLSLToDeviceCompatibleGLSL()
-	UE_DEPRECATED(5.4, "FOpenGLFrontend is deprecated and derived implementations no longer possible; OpenGL shader compilation should only be triggered from within the shader format code.")
-	TSharedPtr<ANSICHAR> PrepareCodeForOfflineCompilation(GLSLVersion ShaderVersion, EShaderFrequency Frequency, const ANSICHAR* InShaderSource) const;
-
-	UE_DEPRECATED(5.4, "FOpenGLFrontend is deprecated and derived implementations no longer possible; OpenGL shader compilation should only be triggered from within the shader format code.")
-	void CompileOffline(const FShaderCompilerInput& ShaderInput, FShaderCompilerOutput &Output, const GLSLVersion ShaderVersion, const ANSICHAR *InShaderSource);
-
-	// based on ShaderVersion decides what platform specific compiler to use
-	UE_DEPRECATED(5.4, "FOpenGLFrontend is deprecated and derived implementations no longer possible; OpenGL shader compilation should only be triggered from within the shader format code.")
-	void PlatformCompileOffline(const FShaderCompilerInput &Input, FShaderCompilerOutput& ShaderOutput, const ANSICHAR* ShaderSource, const GLSLVersion ShaderVersion);
-};
diff --git a/Engine/Source/Runtime/Core/Private/HAL/ConsoleManager.cpp b/Engine/Source/Runtime/Core/Private/HAL/ConsoleManager.cpp
--- a/Engine/Source/Runtime/Core/Private/HAL/ConsoleManager.cpp
+++ b/Engine/Source/Runtime/Core/Private/HAL/ConsoleManager.cpp
@@ -4175,14 +4175,6 @@ static TAutoConsoleVariable<int32> CVarDisableOpenGLTextureStreamingSupport(
 		 "  1 = Texture streaming will be disabled."),
 	ECVF_ReadOnly);
 
-// Moved here from OpenGLRHI module to make sure its always accessible on all platforms
-static FAutoConsoleVariable CVarOpenGLUseEmulatedUBs(
-	TEXT("OpenGL.UseEmulatedUBs"),
-	1,
-	TEXT("If true, enable using emulated uniform buffers on OpenGL Mobile mode."),
-	ECVF_ReadOnly
-	);
-
 static TAutoConsoleVariable<int32> CVarAndroidOverrideExternalTextureSupport(
 	TEXT("r.Android.OverrideExternalTextureSupport"),
 	0,
diff --git a/Engine/Source/Runtime/Engine/Private/ShaderCompiler/ShaderCompiler.cpp b/Engine/Source/Runtime/Engine/Private/ShaderCompiler/ShaderCompiler.cpp
--- a/Engine/Source/Runtime/Engine/Private/ShaderCompiler/ShaderCompiler.cpp
+++ b/Engine/Source/Runtime/Engine/Private/ShaderCompiler/ShaderCompiler.cpp
@@ -2245,14 +2245,6 @@ static TAutoConsoleVariable<int32> CVarD3DForceDXC(
 	TEXT(" 1: Force new compiler for all shaders"),
 	ECVF_ReadOnly);
 
-static TAutoConsoleVariable<int32> CVarOpenGLForceDXC(
-	TEXT("r.OpenGL.ForceDXC"),
-	1,
-	TEXT("Forces DirectX Shader Compiler (DXC) to be used for all OpenGL shaders instead of hlslcc.\n")
-	TEXT(" 0: Disable\n")
-	TEXT(" 1: Force new compiler for all shaders (default)"),
-	ECVF_ReadOnly);
-
 static TAutoConsoleVariable<int32> CVarWarpCulling(
 	TEXT("r.WarpCulling"),
 	0,
diff --git a/Engine/Source/Runtime/RenderCore/Private/Shader.cpp b/Engine/Source/Runtime/RenderCore/Private/Shader.cpp
--- a/Engine/Source/Runtime/RenderCore/Private/Shader.cpp
+++ b/Engine/Source/Runtime/RenderCore/Private/Shader.cpp
@@ -1441,13 +1441,8 @@ bool IsDxcEnabledForPlatform(EShaderPlatform Platform, bool bHlslVersion2021)
 		static const IConsoleVariable* CVar = IConsoleManager::Get().FindConsoleVariable(TEXT("r.D3D.ForceDXC"));
 		return ((CVar && CVar->GetInt() != 0));
 	}
-	if (IsOpenGLPlatform(Platform))
-	{
-		static const IConsoleVariable* CVar = IConsoleManager::Get().FindConsoleVariable(TEXT("r.OpenGL.ForceDXC"));
-		return (bHlslVersion2021 || (CVar && CVar->GetInt() != 0));
-	}
-	// Hlslcc has been removed for Metal and Vulkan. There is only DXC now.
-	if (IsMetalPlatform(Platform) || IsVulkanPlatform(Platform))
+	// Hlslcc has been removed for Metal, Vulkan, and OpenGL backends. There is only DXC now.
+	if (IsMetalPlatform(Platform) || IsVulkanPlatform(Platform) || IsOpenGLPlatform(Platform))
 	{
 		return true;
 	}
@@ -1458,17 +1453,9 @@ bool IsUsingEmulatedUniformBuffers(EShaderPlatform Platform)
 {
 	if (IsOpenGLPlatform(Platform))
 	{
-		// Currently DXC only supports emulated uniform buffers on GLES
-		static const auto CForceDXCVar = IConsoleManager::Get().FindConsoleVariable(TEXT("r.OpenGL.ForceDXC"));
-		if (CForceDXCVar && CForceDXCVar->GetInt() != 0)
-		{
-			return true;
-		}
-
-		static auto* CVar = IConsoleManager::Get().FindTConsoleVariableDataInt(TEXT("OpenGL.UseEmulatedUBs"));
-		return (CVar && CVar->GetValueOnAnyThread() != 0);
+		// DXC only supports emulated uniform buffers on GLES
+		return true;
 	}
-
 	return false;
 }
 
