Fix shader compile problem on linux OpenGLES shader builds

#rnx

diff --git a/Engine/Source/Developer/ShaderFormatOpenGL/Private/OpenGLShaderCompiler.cpp b/Engine/Source/Developer/ShaderFormatOpenGL/Private/OpenGLShaderCompiler.cpp
--- a/Engine/Source/Developer/ShaderFormatOpenGL/Private/OpenGLShaderCompiler.cpp
+++ b/Engine/Source/Developer/ShaderFormatOpenGL/Private/OpenGLShaderCompiler.cpp
@@ -2038,7 +2038,7 @@ struct GLSLCompileParameters
 	CrossCompiler::FShaderConductorOptions* Options;
 
 	FShaderCompilerOutput* Output;
-	TArray<uint32>* SpirvData;
+	TArray<uint32> SpirvData;
 
 	EShaderFrequency Frequency;
 	const ANSICHAR* SPIRV_DummySamplerName;
@@ -2047,7 +2047,7 @@ struct GLSLCompileParameters
 bool GenerateGlslShader(std::string& OutString, GLSLCompileParameters& GLSLCompileParams, ReflectionData& ReflectData, bool bWriteToCCHeader, bool bIsDeferred, bool bEmulatedUBs)
 {
 	const bool bGlslSourceCompileSucceeded = GLSLCompileParams.CompilerContext->CompileSpirvToSourceBuffer(
-		*GLSLCompileParams.Options, *GLSLCompileParams.TargetDesc, GLSLCompileParams.SpirvData->GetData(), GLSLCompileParams.SpirvData->Num() * sizeof(uint32),
+		*GLSLCompileParams.Options, *GLSLCompileParams.TargetDesc, GLSLCompileParams.SpirvData.GetData(), GLSLCompileParams.SpirvData.Num() * sizeof(uint32),
 		[&OutString](const void* Data, uint32 Size)
 		{
 			OutString = std::string(reinterpret_cast<const ANSICHAR*>(Data), Size);
@@ -3108,7 +3108,7 @@ static bool CompileToGlslWithShaderConductor(
 		GLSLCompileParams.Options = &Options;
 
 		GLSLCompileParams.Output = &Output;
-		GLSLCompileParams.SpirvData = &SpirvData;
+		GLSLCompileParams.SpirvData = MoveTemp(SpirvData);
 
 		GLSLCompileParams.Frequency = Frequency;
 		GLSLCompileParams.SPIRV_DummySamplerName = SPIRV_DummySamplerName;
