Fix crash when cooking Android OpenGLES on Linux by replacing swscanf based parser with native FindNextHLSLDefinitionOfType

#rnx

diff --git a/Engine/Source/Developer/ShaderCompilerCommon/Private/ShaderCompilerCommon.cpp b/Engine/Source/Developer/ShaderCompilerCommon/Private/ShaderCompilerCommon.cpp
--- a/Engine/Source/Developer/ShaderCompilerCommon/Private/ShaderCompilerCommon.cpp
+++ b/Engine/Source/Developer/ShaderCompilerCommon/Private/ShaderCompilerCommon.cpp
@@ -360,6 +360,27 @@ const TCHAR* ParseHLSLSymbolName(const TCHAR* SearchString, FString& SymbolName)
 	return Result.GetData() + Result.Len();
 }
 
+FStringView FindNextHLSLDefinitionOfType(FStringView Typename, FStringView StartPos)
+{
+	// handle both the case where identifier for declaration immediately precedes a ; and has whitespace separating the two
+	const TCHAR* NextWhitespace;
+	const TCHAR* NextNonWhitespace;
+	FStringView SymbolName;
+
+	NextWhitespace = FindNextWhitespace(StartPos.GetData());
+	if (NextWhitespace == StartPos.GetData())
+	{
+		NextNonWhitespace = FindNextNonWhitespace(NextWhitespace);
+		SymbolName = ParseHLSLSymbolName<TCHAR, FStringView>(NextNonWhitespace);	
+		NextNonWhitespace = FindNextNonWhitespace(NextNonWhitespace + SymbolName.Len());
+		if (NextNonWhitespace && (*NextNonWhitespace == ';'))
+		{
+			return SymbolName;
+		}
+	}
+	return {};
+}
+
 FStringView UE::ShaderCompilerCommon::RemoveConstantBufferPrefix(FStringView InName)
 {
 	return UE::String::RemoveFromStart(InName, FStringView(UE::ShaderCompilerCommon::kUniformBufferConstantBufferPrefix));
diff --git a/Engine/Source/Developer/ShaderCompilerCommon/Public/ShaderCompilerCommon.h b/Engine/Source/Developer/ShaderCompilerCommon/Public/ShaderCompilerCommon.h
--- a/Engine/Source/Developer/ShaderCompilerCommon/Public/ShaderCompilerCommon.h
+++ b/Engine/Source/Developer/ShaderCompilerCommon/Public/ShaderCompilerCommon.h
@@ -366,6 +366,7 @@ const CharType* FindMatchingClosingBrace(const CharType* OpeningCharPtr) { retur
 
 extern SHADERCOMPILERCOMMON_API const TCHAR* ParseHLSLSymbolName(const TCHAR* SearchString, FString& SymboName);
 extern SHADERCOMPILERCOMMON_API void ParseHLSLTypeName(const TCHAR* SearchString, const TCHAR*& TypeNameStartPtr, const TCHAR*& TypeNameEndPtr);
+extern SHADERCOMPILERCOMMON_API FStringView FindNextHLSLDefinitionOfType(FStringView Typename, FStringView StartPos);
 
 UE_DEPRECATED(5.4, "TEXT macro processing is now handled in PreprocessShader --  TransformStringIntoCharacterArray is a no-op")
 inline void TransformStringIntoCharacterArray(FString& PreprocessedShaderSource, TArray<FShaderDiagnosticData>* OutDiagnosticDatas = nullptr) {}
diff --git a/Engine/Source/Developer/ShaderFormatOpenGL/Private/OpenGLShaderCompiler.cpp b/Engine/Source/Developer/ShaderFormatOpenGL/Private/OpenGLShaderCompiler.cpp
--- a/Engine/Source/Developer/ShaderFormatOpenGL/Private/OpenGLShaderCompiler.cpp
+++ b/Engine/Source/Developer/ShaderFormatOpenGL/Private/OpenGLShaderCompiler.cpp
@@ -8,6 +8,7 @@
 #include "Misc/FileHelper.h"
 #include "Misc/Paths.h"
 #include "Serialization/MemoryWriter.h"
+#include "ShaderCompilerCommon.h"
 #include "ShaderCompilerDefinitions.h"
 #include "ShaderFormatOpenGL.h"
 #include "ShaderParameterParser.h"
@@ -3050,32 +3051,19 @@ static bool CompileToGlslWithShaderConductor(
 
 		TSet<FString> ExternalTextures;
 		int32 Pos = 0;
-#if !PLATFORM_MAC
-		TCHAR TextureExternalName[257] = {};
-		TCHAR NextToken[2] = {};
-#else
-		ANSICHAR TextureExternalName[257] = {};
-		ANSICHAR NextToken[2] = {};
-#endif
+		static constexpr FStringView TextExternal = TEXTVIEW("TextureExternal");
+		
 		do
 		{
-			Pos = PreprocessedShader.Find(TEXT("TextureExternal"), ESearchCase::CaseSensitive, ESearchDir::FromStart, Pos + 15);
+			Pos = PreprocessedShader.Find(TextExternal, ESearchCase::CaseSensitive, ESearchDir::FromStart, Pos + TextExternal.Len());
 			if (Pos != INDEX_NONE)
 			{
-				// handle both the case where identifier for TextureExternal declaration immediately precedes a ; and has whitespace separating the two
-#if PLATFORM_WINDOWS
-				if (swscanf_s(&PreprocessedShader[Pos], TEXT("TextureExternal %ls %ls"), TextureExternalName, 257, NextToken, 2))
-#elif PLATFORM_MAC
-				if (sscanf(TCHAR_TO_ANSI(&PreprocessedShader[Pos]), "TextureExternal %256s %1s", TextureExternalName, NextToken))
-#else // PLATFORM_LINUX
-				if (swscanf(TCHAR_TO_WCHAR(&PreprocessedShader[Pos]), L"TextureExternal %256ls %1ls", TextureExternalName, NextToken))
-#endif
+				FStringView TextureExternalName;
+
+				TextureExternalName = FindNextHLSLDefinitionOfType(FStringView(&PreprocessedShader[Pos]), FStringView(&PreprocessedShader[Pos+TextExternal.Len()]));
+				if (!TextureExternalName.IsEmpty())
 				{
-					FString Name = TextureExternalName;
-					if (Name.RemoveFromEnd(TEXT(";")) || (NextToken[0] == TEXT(';')))
-					{
-						ExternalTextures.Add(Name + TEXT("Sampler"));
-					}
+					ExternalTextures.Add(FString(TextureExternalName) + TEXT("Sampler"));	
 				}
 			}
 		}
