This fixes shader OpenGL ES shader link errors that prevent pretty much any project that use GL to work.

The option bPreserveStorageInput of current version of SPIRV-Tools does not work as it used to and unused fragment shader inputs are now stripped by RemoveUnusedInterfaceVariablesPass.
This fix does not address this but instead disables the input location remapping of OpenGLShaderCompiler for fragment shaders so that even when inputs are stripped, the remaining inputs still have their original location indices that match the vertex shader outputs.

diff --git a/Engine/Source/Developer/ShaderFormatOpenGL/Private/OpenGLShaderCompiler.cpp b/Engine/Source/Developer/ShaderFormatOpenGL/Private/OpenGLShaderCompiler.cpp
--- a/Engine/Source/Developer/ShaderFormatOpenGL/Private/OpenGLShaderCompiler.cpp
+++ b/Engine/Source/Developer/ShaderFormatOpenGL/Private/OpenGLShaderCompiler.cpp
@@ -1818,7 +1818,7 @@ void ParseReflectionData(const FShaderCompilerInput& ShaderInput, CrossCompiler:
 		ReflectionBindings.GatherInputAttributes(Reflection);
 		for (SpvReflectInterfaceVariable* Var : ReflectionBindings.InputAttributes)
 		{
-			if (Var->storage_class == SpvStorageClassInput && Var->built_in == -1)
+			if (Var->storage_class == SpvStorageClassInput && Var->built_in == -1 && Frequency == SF_Vertex)
 			{
 				unsigned Location = Var->location;
 				unsigned SemanticIndex = Location;
@@ -2874,9 +2874,8 @@ static bool CompileToGlslWithShaderConductor(
 	// Initialize compilation options for ShaderConductor
 	CrossCompiler::FShaderConductorOptions Options;
 	Options.bDisableScalarBlockLayout = true;
-	Options.bRemapAttributeLocations = true;
-	Options.bPreserveStorageInput = true;
-    Options.bForceStorageImageFormat = true;
+	Options.bRemapAttributeLocations = (Frequency == SF_Vertex);
+	Options.bForceStorageImageFormat = true;
 	Options.bWarningsAsErrors = Input.Environment.CompilerFlags.Contains(CFLAG_WarningsAsErrors);
 	
 	// Enable HLSL 2021 if specified
