Fix the GBufferA format for the mobile deferred with pixel local storage .

diff --git a/Engine/Source/Developer/ShaderFormatOpenGL/Private/OpenGLShaderCompiler.cpp b/Engine/Source/Developer/ShaderFormatOpenGL/Private/OpenGLShaderCompiler.cpp
--- a/Engine/Source/Developer/ShaderFormatOpenGL/Private/OpenGLShaderCompiler.cpp
+++ b/Engine/Source/Developer/ShaderFormatOpenGL/Private/OpenGLShaderCompiler.cpp
@@ -2762,7 +2762,7 @@ bool GenerateDeferredMobileShaders(std::string& GlslSource, GLSLCompileParameter
 		static const FString GBufferOutputNames[] =
 		{
 			TEXT("rgb10_a2 out.var.SV_Target0"),
-			TEXT("rgba8 out.var.SV_Target1"),
+			TEXT("rgb10_a2 out.var.SV_Target1"),
 			TEXT("rgba8 out.var.SV_Target2"),
 			TEXT("rgba8 out.var.SV_Target3"),
 		};
@@ -2770,7 +2770,7 @@ bool GenerateDeferredMobileShaders(std::string& GlslSource, GLSLCompileParameter
 		static const FString GBufferInputNames[] =
 		{
 			TEXT("rgb10_a2 GENERATED_SubpassFetchAttachment0"),
-			TEXT("rgba8 GENERATED_SubpassFetchAttachment1"),
+			TEXT("rgb10_a2 GENERATED_SubpassFetchAttachment1"),
 			TEXT("rgba8 GENERATED_SubpassFetchAttachment2"),
 			TEXT("rgba8 GENERATED_SubpassFetchAttachment3"),
 		};
@@ -2778,7 +2778,7 @@ bool GenerateDeferredMobileShaders(std::string& GlslSource, GLSLCompileParameter
 		static const FString GBufferInOutNames[] =
 		{
 			TEXT("rgb10_a2 GENERATED_SubpassFetchAttachment0 out.var.SV_Target0"),
-			TEXT("rgba8 GENERATED_SubpassFetchAttachment1 out.var.SV_Target1"),
+			TEXT("rgb10_a2 GENERATED_SubpassFetchAttachment1 out.var.SV_Target1"),
 			TEXT("rgba8 GENERATED_SubpassFetchAttachment2 out.var.SV_Target2"),
 			TEXT("rgba8 GENERATED_SubpassFetchAttachment3 out.var.SV_Target3"),
 		};
