[Android] Fixing up GLSL shader compilation to once more allow usage of the OES External Texture extension

diff --git a/Engine/Source/Developer/ShaderFormatOpenGL/Private/OpenGLShaderCompiler.cpp b/Engine/Source/Developer/ShaderFormatOpenGL/Private/OpenGLShaderCompiler.cpp
--- a/Engine/Source/Developer/ShaderFormatOpenGL/Private/OpenGLShaderCompiler.cpp
+++ b/Engine/Source/Developer/ShaderFormatOpenGL/Private/OpenGLShaderCompiler.cpp
@@ -3040,7 +3040,7 @@ static bool CompileToGlslWithShaderConductor(
 				TextureExternalName = FindNextHLSLDefinitionOfType(FStringView(&PreprocessedShader[Pos]), FStringView(&PreprocessedShader[Pos+TextExternal.Len()]));
 				if (!TextureExternalName.IsEmpty())
 				{
-					ExternalTextures.Add(FString(TextureExternalName) + TEXT("Sampler"));	
+					ExternalTextures.Add(FString(TextureExternalName));
 				}
 			}
 		}
@@ -3049,9 +3049,11 @@ static bool CompileToGlslWithShaderConductor(
 		// Define type renaming callback after all external texture types have been gathered
 		TargetDesc.VariableTypeRenameCallback = [&ExternalTextures](const FAnsiStringView& VariableName, const FAnsiStringView& TypeName, FString& OutRenamedTypeName) -> bool
 		{
+			auto WideVarName = StringCast<TCHAR>(VariableName.GetData());
 			for (const FString& ExternalTex : ExternalTextures)
 			{
-				if (FCStringWide::Strstr(ANSI_TO_TCHAR(VariableName.GetData()), *ExternalTex))
+				if (FCStringWide::Strstr(WideVarName.Get(), *ExternalTex) &&
+					FCStringWide::Strstr(WideVarName.Get(), TEXT("_SamplerP")))
 				{
 					OutRenamedTypeName = TEXT("samplerExternalOES");
 					return true;
