[Shaders] Remove unused platform specific identifiers from GLES shader backend.

#rnx

diff --git a/Engine/Source/Developer/ShaderFormatOpenGL/Private/OpenGLShaderCompiler.cpp b/Engine/Source/Developer/ShaderFormatOpenGL/Private/OpenGLShaderCompiler.cpp
--- a/Engine/Source/Developer/ShaderFormatOpenGL/Private/OpenGLShaderCompiler.cpp
+++ b/Engine/Source/Developer/ShaderFormatOpenGL/Private/OpenGLShaderCompiler.cpp
@@ -2933,14 +2933,6 @@ static bool CompileToGlslWithShaderConductor(
 
 		switch (Version)
 		{
-		case GLSL_SWITCH_FORWARD:
-			TargetDesc.Language = CrossCompiler::EShaderConductorLanguage::Essl;
-			TargetDesc.Version = 320;
-			break;
-		case GLSL_SWITCH:
-			TargetDesc.Language = CrossCompiler::EShaderConductorLanguage::Glsl;
-			TargetDesc.Version = 430;
-			break;
 		case GLSL_150_ES3_1:
 		case GLSL_ES3_1_ANDROID:
 		default:
diff --git a/Engine/Source/Developer/ShaderFormatOpenGL/Public/ShaderFormatOpenGL.h b/Engine/Source/Developer/ShaderFormatOpenGL/Public/ShaderFormatOpenGL.h
--- a/Engine/Source/Developer/ShaderFormatOpenGL/Public/ShaderFormatOpenGL.h
+++ b/Engine/Source/Developer/ShaderFormatOpenGL/Public/ShaderFormatOpenGL.h
@@ -33,8 +33,8 @@ enum GLSLVersion
 	GLSL_ES2_IOS_REMOVED,
 	GLSL_310_ES_EXT_REMOVED,
 	GLSL_ES3_1_ANDROID,
-	GLSL_SWITCH,
-	GLSL_SWITCH_FORWARD,
+	GLSL_SWITCH UE_DEPRECATED(5.5, "no longer supported in this shader backend"),
+	GLSL_SWITCH_FORWARD UE_DEPRECATED(5.5, "no longer supported in this shader backend"),
 
 	GLSL_MAX
 };
