Fix bug in HLSL to GLSL conversion.
Before the change we iterate over buffer names and string replace.
But this can fail when we have buffer names like MyBuffer followed by MyBufferPlusSuffix.
Because when replacing MyBuffer the rename search matches MyBufferPlusSuffix.
To fix the issue we presort the array of buffer names so that MyBufferPlusSuffix always gets replaced first.
Care needs to be taken that the original order of UAVIndex values continues to be respected.

diff --git a/Engine/Source/Developer/ShaderFormatOpenGL/Private/OpenGLShaderCompiler.cpp b/Engine/Source/Developer/ShaderFormatOpenGL/Private/OpenGLShaderCompiler.cpp
--- a/Engine/Source/Developer/ShaderFormatOpenGL/Private/OpenGLShaderCompiler.cpp
+++ b/Engine/Source/Developer/ShaderFormatOpenGL/Private/OpenGLShaderCompiler.cpp
@@ -2367,9 +2367,20 @@ bool GenerateGlslShader(std::string& OutString, GLSLCompileParameters& GLSLCompi
 		}
 
 		// StructuredBuffers, rename as ci0 format
+		// Presort the rename order to avoid search collisions for names like MyBuffer and MyBufferWithSuffix.
+		using FIndexNamePair = TPair<int32, const std::string&>;
+		TArray<FIndexNamePair> SortedBuffers;
+		SortedBuffers.Reserve(ReflectData.StructuredBuffers.Num());
 		for (const std::string& SBuffer : ReflectData.StructuredBuffers)
 		{
-			std::string NewSBufferName = FrequencyPrefix + std::string("i") + std::to_string(UAVIndex);
+			SortedBuffers.Add(FIndexNamePair(UAVIndex++, SBuffer));
+		}
+		Algo::Sort(SortedBuffers, [](const FIndexNamePair& A, const FIndexNamePair& B){ return A.Value > B.Value; });
+		
+		for (const FIndexNamePair& SortedEntry : SortedBuffers)
+		{
+			const std::string& SBuffer = SortedEntry.Value;
+			std::string NewSBufferName = FrequencyPrefix + std::string("i") + std::to_string(SortedEntry.Key);
 			std::string NewSBufferVarName = NewSBufferName + std::string("_VAR");
 
 			std::string SearchString = "} " + SBuffer;
@@ -2399,8 +2410,6 @@ bool GenerateGlslShader(std::string& OutString, GLSLCompileParameters& GLSLCompi
 				OutString.replace(CurPos, ReplacementSubStr.size() + 4, NewSBufferName);
 				CurPos = OutString.find(ReplacementSubStr + ".");
 			}
-
-			UAVIndex++;
 		}
 
 		for (const auto& pair : ReflectData.UniformVarNames)
