[Shaders] Use FAnsiString in GL backend instead of TSharedPtr to store shader source as ANSI string.

This fixes static analysis issues when a temporary string is allocated with "new" but deleted with "delete" instead of "delete[]" via TSharedPtr.
The solution is to use FAnsiString which also avoids the unnecessary conversion from ANSI to TCHAR and back again.

#rnx

diff --git a/Engine/Source/Developer/ShaderFormatOpenGL/Private/OpenGLShaderCompiler.cpp b/Engine/Source/Developer/ShaderFormatOpenGL/Private/OpenGLShaderCompiler.cpp
--- a/Engine/Source/Developer/ShaderFormatOpenGL/Private/OpenGLShaderCompiler.cpp
+++ b/Engine/Source/Developer/ShaderFormatOpenGL/Private/OpenGLShaderCompiler.cpp
@@ -555,8 +555,8 @@ struct FDeviceCapabilities
 
 static bool PlatformSupportsOfflineCompilationInternal(const GLSLVersion ShaderVersion);
 static void FillDeviceCapsOfflineCompilationInternal(struct FDeviceCapabilities& Capabilities, const GLSLVersion ShaderVersion);
-static bool MoveHashLines(FString& Destination, FString &Source);
-static TSharedPtr<ANSICHAR> PrepareCodeForOfflineCompilationInternal(const GLSLVersion ShaderVersion, EShaderFrequency Frequency, const ANSICHAR* InShaderSource);
+static bool MoveHashLines(FAnsiString& Destination, FAnsiString& Source);
+static FAnsiString PrepareCodeForOfflineCompilationInternal(const GLSLVersion ShaderVersion, EShaderFrequency Frequency, const ANSICHAR* InShaderSource);
 static void PlatformCompileOfflineInternal(const FShaderCompilerInput& Input, FShaderCompilerOutput& ShaderOutput, const ANSICHAR* ShaderSource, const GLSLVersion ShaderVersion);
 static void CompileOfflineInternal(const FShaderCompilerInput& Input, FShaderCompilerOutput& Output, const GLSLVersion ShaderVersion, const ANSICHAR* InShaderSource);
 
@@ -3160,7 +3160,7 @@ static void FillDeviceCapsOfflineCompilationInternal(struct FDeviceCapabilities&
 	}
 }
 
-static bool MoveHashLines(FString& Destination, FString &Source)
+static bool MoveHashLines(FAnsiString& Destination, FAnsiString& Source)
 {
 	int32 Index = 0;
 	int32 LineStart = 0;
@@ -3169,7 +3169,7 @@ static bool MoveHashLines(FString& Destination, FString &Source)
 	while (Index != INDEX_NONE && !bFound)
 	{
 		LineStart = Index;
-		Index = Source.Find(TEXT("\n"), ESearchCase::IgnoreCase, ESearchDir::FromStart, Index);
+		Index = Source.Find("\n", ESearchCase::IgnoreCase, ESearchDir::FromStart, Index);
 
 		for (int32 i = LineStart; i < Index; ++i)
 		{
@@ -3197,25 +3197,25 @@ static bool MoveHashLines(FString& Destination, FString &Source)
 	return bFound;
 }
 
-static TSharedPtr<ANSICHAR> PrepareCodeForOfflineCompilationInternal(const GLSLVersion ShaderVersion, EShaderFrequency Frequency, const ANSICHAR* InShaderSource)
+static FAnsiString PrepareCodeForOfflineCompilationInternal(const GLSLVersion ShaderVersion, EShaderFrequency Frequency, const ANSICHAR* InShaderSource)
 {
-	FString OriginalShaderSource(ANSI_TO_TCHAR(InShaderSource));
-	FString StrOutSource;
+	FAnsiString OriginalShaderSource(InShaderSource);
+	FAnsiString StrOutSource;
 
 	FDeviceCapabilities Capabilities;
 	FillDeviceCapsOfflineCompilationInternal(Capabilities, ShaderVersion);
 
 	// Whether we need to emit mobile multi-view code or not.
-	const bool bEmitMobileMultiView = OriginalShaderSource.Find(TEXT("gl_ViewID_OVR")) != INDEX_NONE;
+	const bool bEmitMobileMultiView = OriginalShaderSource.Find("gl_ViewID_OVR") != INDEX_NONE;
 
 	// Whether we need to emit texture external code or not.
-	const bool bEmitTextureExternal = OriginalShaderSource.Find(TEXT("samplerExternalOES")) != INDEX_NONE;
+	const bool bEmitTextureExternal = OriginalShaderSource.Find("samplerExternalOES") != INDEX_NONE;
 
 	if (Capabilities.TargetPlatform == EPlatformType::Android || Capabilities.TargetPlatform == EPlatformType::Web)
 	{
-		const TCHAR *ES320Version = TEXT("#version 320 es");
+		const ANSICHAR* ES320Version = "#version 320 es";
 		StrOutSource.Append(ES320Version);
-		StrOutSource.Append(TEXT("\n"));
+		StrOutSource.Append("\n");
 		OriginalShaderSource.RemoveFromStart(ES320Version);
 	}
 
@@ -3225,15 +3225,15 @@ static TSharedPtr<ANSICHAR> PrepareCodeForOfflineCompilationInternal(const GLSLV
 	{
 		MoveHashLines(StrOutSource, OriginalShaderSource);
 
-		StrOutSource.Append(TEXT("\n\n"));
-		StrOutSource.Append(TEXT("#extension GL_OVR_multiview2 : enable\n"));
-		StrOutSource.Append(TEXT("\n\n"));
+		StrOutSource.Append("\n\n");
+		StrOutSource.Append("#extension GL_OVR_multiview2 : enable\n");
+		StrOutSource.Append("\n\n");
 	}
 
 	if (bEmitTextureExternal)
 	{
 		MoveHashLines(StrOutSource, OriginalShaderSource);
-		StrOutSource.Append(TEXT("#define samplerExternalOES sampler2D\n"));
+		StrOutSource.Append("#define samplerExternalOES sampler2D\n");
 	}
 
 	// Move version tag & extensions before beginning all other operations
@@ -3242,37 +3242,33 @@ static TSharedPtr<ANSICHAR> PrepareCodeForOfflineCompilationInternal(const GLSLV
 	// OpenGL SM5 shader platforms require location declarations for the layout, but don't necessarily use SSOs
 	if (Capabilities.TargetPlatform == EPlatformType::Desktop)
 	{
-		StrOutSource.Append(TEXT("#define INTERFACE_BLOCK(Pos, Interp, Modifiers, Semantic, PreType, PostType) layout(location=Pos) Interp Modifiers struct { PreType PostType; }\n"));
+		StrOutSource.Append("#define INTERFACE_BLOCK(Pos, Interp, Modifiers, Semantic, PreType, PostType) layout(location=Pos) Interp Modifiers struct { PreType PostType; }\n");
 	}
 	else
 	{
-		StrOutSource.Append(TEXT("#define INTERFACE_BLOCK(Pos, Interp, Modifiers, Semantic, PreType, PostType) layout(location=Pos) Modifiers Semantic { PreType PostType; }\n"));
+		StrOutSource.Append("#define INTERFACE_BLOCK(Pos, Interp, Modifiers, Semantic, PreType, PostType) layout(location=Pos) Modifiers Semantic { PreType PostType; }\n");
 	}
 
 	if (Capabilities.TargetPlatform == EPlatformType::Desktop)
 	{
 		// If we're running <= featurelevel es3.1 shaders then enable this extension which adds support for uintBitsToFloat etc.
-		if (StrOutSource.Contains(TEXT("#version 150")))
+		if (StrOutSource.Contains("#version 150"))
 		{
-			StrOutSource.Append(TEXT("\n\n"));
-			StrOutSource.Append(TEXT("#extension GL_ARB_gpu_shader5 : enable\n"));
-			StrOutSource.Append(TEXT("\n\n"));
+			StrOutSource.Append("\n\n");
+			StrOutSource.Append("#extension GL_ARB_gpu_shader5 : enable\n");
+			StrOutSource.Append("\n\n");
 		}
 	}
 
-	StrOutSource.Append(TEXT("#define HLSLCC_DX11ClipSpace 1 \n"));
+	StrOutSource.Append("#define HLSLCC_DX11ClipSpace 1 \n");
 
 	// Append the possibly edited shader to the one we will compile.
 	// This is to make it easier to debug as we can see the whole
 	// shader source.
-	StrOutSource.Append(TEXT("\n\n"));
+	StrOutSource.Append("\n\n");
 	StrOutSource.Append(OriginalShaderSource);
 
-	const int32 SourceLen = StrOutSource.Len();
-	TSharedPtr<ANSICHAR> RetShaderSource = MakeShareable(new ANSICHAR[SourceLen + 1]);
-	FCStringAnsi::Strncpy(RetShaderSource.Get(), TCHAR_TO_ANSI(*StrOutSource), SourceLen + 1);
-
-	return RetShaderSource;
+	return StrOutSource;
 }
 
 static bool PlatformSupportsOfflineCompilationInternal(const GLSLVersion ShaderVersion)
@@ -3289,9 +3285,9 @@ static void CompileOfflineInternal(const FShaderCompilerInput& Input, FShaderCom
 		return;
 	}
 
-	TSharedPtr<ANSICHAR> ShaderSource = PrepareCodeForOfflineCompilationInternal(ShaderVersion, (EShaderFrequency)Input.Target.Frequency, InShaderSource);
+	FAnsiString ShaderSource = PrepareCodeForOfflineCompilationInternal(ShaderVersion, (EShaderFrequency)Input.Target.Frequency, InShaderSource);
 
-	PlatformCompileOfflineInternal(Input, Output, ShaderSource.Get(), ShaderVersion);
+	PlatformCompileOfflineInternal(Input, Output, *ShaderSource, ShaderVersion);
 }
 
 static void PlatformCompileOfflineInternal(const FShaderCompilerInput& Input, FShaderCompilerOutput& ShaderOutput, const ANSICHAR* ShaderSource, const GLSLVersion ShaderVersion)
