[Shaders] Fix sampler limit in OpenGLShaderCompiler.

The GLSL backend shares texture slots between SRV and UAV to emulate atomics with buffers.
This change expands the TextureIndices bitmask to 64-bit to account for up to 32 sampler slots (SRV) plus image slots (UAV).
The bitmask manipulation already operated on 64-bit types, so this might have been an oversight with the introduction of this bitmask.

#rnx

diff --git a/Engine/Source/Developer/ShaderFormatOpenGL/Private/OpenGLShaderCompiler.cpp b/Engine/Source/Developer/ShaderFormatOpenGL/Private/OpenGLShaderCompiler.cpp
--- a/Engine/Source/Developer/ShaderFormatOpenGL/Private/OpenGLShaderCompiler.cpp
+++ b/Engine/Source/Developer/ShaderFormatOpenGL/Private/OpenGLShaderCompiler.cpp
@@ -1409,7 +1409,9 @@ void ParseReflectionData(const FShaderCompilerInput& ShaderInput, CrossCompiler:
 
 	uint32 BufferIndices = 0xffffffff;
 	uint32 UAVIndices = 0xffffffff;
-	uint32 TextureIndices = 0xffffffff;
+	// The GLSL backend shares texture slots between SRV and UAV, so we need more than 32 slots.
+	// Exceeding the SRV limit of 32 samplers should not be handled by this function, so use 64 bit.
+	uint64 TextureIndices = 0xffffffffffffffffllu;
 	uint32 UBOIndices = 0xffffffff;
 	uint32 SamplerIndices = 0xffffffff;
 
@@ -1470,7 +1472,7 @@ void ParseReflectionData(const FShaderCompilerInput& ShaderInput, CrossCompiler:
 	for (auto const& Binding : ReflectionBindings.TBufferSRVs)
 	{
 		check(TextureIndices);
-		uint32 Index = FPlatformMath::CountTrailingZeros(TextureIndices);
+		uint32 Index = static_cast<uint32>(FPlatformMath::CountTrailingZeros64(TextureIndices));
 
 		// No support for 3-component types in dxc/SPIRV/MSL - need to expose my workarounds there too
 		TextureIndices &= ~(1llu << uint64(Index));
@@ -1626,7 +1628,7 @@ void ParseReflectionData(const FShaderCompilerInput& ShaderInput, CrossCompiler:
 	for (auto const& Binding : ReflectionBindings.TextureSRVs)
 	{
 		check(TextureIndices);
-		uint32 Index = FPlatformMath::CountTrailingZeros64(TextureIndices);
+		uint32 Index = static_cast<uint32>(FPlatformMath::CountTrailingZeros64(TextureIndices));
 		TextureIndices &= ~(1llu << uint64(Index));
 
 		ReflectionOut.Textures.Add(UTF8_TO_TCHAR(Binding->name));
